AnimNextPoseSearch - HistoryCollectorTrait, MotionMatchingTrait

diff --git a/Engine/Plugins/Experimental/AnimNextPoseSearch/AnimNextPoseSearch.uplugin b/Engine/Plugins/Experimental/AnimNextPoseSearch/AnimNextPoseSearch.uplugin
new file mode 100644
--- /dev/null
+++ b/Engine/Plugins/Experimental/AnimNextPoseSearch/AnimNextPoseSearch.uplugin
@@ -0,0 +1,41 @@
+{
+	"FileVersion" : 3,
+	"Version" : 1,
+	"VersionName" : "0.1",
+	"FriendlyName" : "Anim Next Pose Search",
+	"Description" : "Pose Search integration for Anim Next.",
+	"Category" : "Animation",
+	"CreatedBy" : "Epic Games, Inc.",
+	"CreatedByURL" : "https://epicgames.com",
+	"DocsURL" : "",
+	"MarketplaceURL" : "",
+	"SupportURL" : "",
+	"EnabledByDefault" : false,
+	"CanContainContent" : false,
+	"IsBetaVersion": true,
+	"IsExperimentalVersion": true,
+	"Installed" : false,
+	"Plugins":
+	[
+		{
+			"Name": "AnimNext",
+			"Enabled": true
+		},
+		{
+			"Name": "AnimNextAnimGraph",
+			"Enabled": true
+		},
+		{
+			"Name": "PoseSearch",
+			"Enabled": true
+		}
+	],
+	"Modules" :
+	[
+		{
+			"Name" : "AnimNextPoseSearch",
+			"Type" : "Runtime",
+			"LoadingPhase" : "Default"
+		},
+	]
+}
diff --git a/Engine/Plugins/Experimental/AnimNextPoseSearch/Source/AnimNextPoseSearch/AnimNextPoseSearch.Build.cs b/Engine/Plugins/Experimental/AnimNextPoseSearch/Source/AnimNextPoseSearch/AnimNextPoseSearch.Build.cs
new file mode 100644
--- /dev/null
+++ b/Engine/Plugins/Experimental/AnimNextPoseSearch/Source/AnimNextPoseSearch/AnimNextPoseSearch.Build.cs
@@ -0,0 +1,24 @@
+// Copyright Epic Games, Inc. All Rights Reserved.
+
+namespace UnrealBuildTool.Rules
+{
+	public class AnimNextPoseSearch : ModuleRules
+	{
+		public AnimNextPoseSearch(ReadOnlyTargetRules Target) : base(Target)
+		{
+			PublicDependencyModuleNames.AddRange(
+				new string[]
+				{
+					"AnimNext",
+					"AnimNextAnimGraph",
+					"ControlRig",
+					"Core",
+					"CoreUObject",
+					"Engine",
+					"PoseSearch",
+					"RigVM"
+				}
+			);
+		}
+	}
+}
\ No newline at end of file
diff --git a/Engine/Plugins/Experimental/AnimNextPoseSearch/Source/AnimNextPoseSearch/Private/HistoryCollectorTrait.cpp b/Engine/Plugins/Experimental/AnimNextPoseSearch/Source/AnimNextPoseSearch/Private/HistoryCollectorTrait.cpp
new file mode 100644
--- /dev/null
+++ b/Engine/Plugins/Experimental/AnimNextPoseSearch/Source/AnimNextPoseSearch/Private/HistoryCollectorTrait.cpp
@@ -0,0 +1,182 @@
+// Copyright Epic Games, Inc. All Rights Reserved.
+
+#include "HistoryCollectorTrait.h"
+#include "Animation/AnimNodeBase.h"
+#include "AnimationRuntime.h"
+#include "EvaluationVM/EvaluationVM.h"
+#include "GenerationTools.h"
+
+FAnimNextHistoryCollectorTask FAnimNextHistoryCollectorTask::Make(UE::PoseSearch::FPoseHistory* PoseHistory, const FAnimNextHistoryCollectorTraitSharedData* SharedData, bool bStoreScales, float DeltaTime)
+{
+	check(PoseHistory);
+	FAnimNextHistoryCollectorTask Task;
+	Task.PoseHistory = PoseHistory;
+	Task.SharedData = SharedData;
+	Task.bStoreScales = bStoreScales;
+	Task.DeltaTime = DeltaTime;
+	return Task;
+}
+
+namespace UE::AnimNext
+{
+	static void RemapPoseInternal(const FLODPoseStack& SourcePose, FCompactPose& TargetPose)
+	{
+		//SCOPE_CYCLE_COUNTER(STAT_AnimNext_RemapPose);
+
+		const FBoneContainer& BoneContainer = TargetPose.GetBoneContainer();
+		const FReferencePose& RefPose = SourcePose.GetRefPose();
+		const TArrayView<const FBoneIndexType> LODBoneIndexes = RefPose.GetLODBoneIndexToMeshBoneIndexMap(SourcePose.LODLevel);
+		const int32 NumLODBones = LODBoneIndexes.Num();
+
+		check(SourcePose.GetNumBones() == NumLODBones);
+
+		for (int32 LODBoneIndex = 0; LODBoneIndex < NumLODBones; ++LODBoneIndex)
+		{
+			// Reference pose holds a list of skeletal mesh bone indices sorted by LOD
+			const FMeshPoseBoneIndex MeshBoneIndex(LODBoneIndexes[LODBoneIndex]);
+
+			// Remap our skeletal mesh bone index into the skeleton bone index we output for
+			const FSkeletonPoseBoneIndex SkeletonBoneIndex = BoneContainer.GetSkeletonPoseIndexFromMeshPoseIndex(MeshBoneIndex);
+			if (SkeletonBoneIndex.IsValid())
+			{
+				// Remap our skeleton bone index into the compact pose bone index we output for
+				const FCompactPoseBoneIndex CompactBoneIndex = BoneContainer.GetCompactPoseIndexFromSkeletonPoseIndex(SkeletonBoneIndex);
+
+				if (ensure(CompactBoneIndex.IsValid()))	// We expect the skeleton bone to map to a valid compact pose bone
+				{
+					TargetPose[CompactBoneIndex] = SourcePose.LocalTransformsView[LODBoneIndex];
+				}
+				else
+				{
+					// This bone is part of the LOD but isn't part of the required bones
+				}
+			}
+			else
+			{
+				// This bone is part of the LOD but isn't part of the required bones
+			}
+		}
+	}
+}
+
+void FAnimNextHistoryCollectorTask::Execute(UE::AnimNext::FEvaluationVM& VM) const
+{
+	using namespace UE::AnimNext;
+
+	if (EnumHasAnyFlags(VM.GetFlags(), EEvaluationFlags::Bones))
+	{
+		if (const TUniquePtr<FKeyframeState>* Keyframe = VM.PeekValue<TUniquePtr<FKeyframeState>>(KEYFRAME_STACK_NAME, 0))
+		{
+			FLODPoseStack& Pose = (*Keyframe)->Pose;
+			const USkeleton* Skeleton = Pose.RefPose->Skeleton.Get();
+			check(Skeleton);
+
+			bool bNeedsReset = false;
+			bool bCacheBones = true;
+
+			if (SharedData->PoseCount != PoseHistory->GetMaxNumPoses() || SharedData->SamplingInterval != PoseHistory->GetSamplingInterval())
+			{
+				PoseHistory->Initialize_AnyThread(SharedData->PoseCount, SharedData->SamplingInterval);
+			}
+
+			FMemMark Mark(FMemStack::Get());
+			TArray<uint16> BoneIndicesWithParents;
+			BoneIndicesWithParents.Add(UE::PoseSearch::RootBoneIndexType);
+
+			for (const FBoneReference& CollectedBone : SharedData->CollectedBones)
+			{
+				if (CollectedBone.BoneName != NAME_None)
+				{
+					FBoneReference TempCollectedBone = CollectedBone;
+					TempCollectedBone.Initialize(Skeleton);
+					if (TempCollectedBone.HasValidSetup())
+					{
+						BoneIndicesWithParents.AddUnique(TempCollectedBone.BoneIndex);
+					}
+				}
+			}
+
+			// Build separate index array with parent indices guaranteed to be present. Sort for EnsureParentsPresent.
+			BoneIndicesWithParents.Sort();
+			FAnimationRuntime::EnsureParentsPresent(BoneIndicesWithParents, Skeleton->GetReferenceSkeleton());
+			
+			FBoneContainer BoneContainer;
+			BoneContainer.InitializeTo(BoneIndicesWithParents, UE::Anim::FCurveFilterSettings(UE::Anim::ECurveFilterMode::DisallowAll), *const_cast<USkeleton*>(Skeleton));
+
+			FCompactPose LocalSpacePose;
+			LocalSpacePose.SetBoneContainer(&BoneContainer);
+
+			RemapPoseInternal(Pose, LocalSpacePose);
+
+			FCSPose<FCompactPose> ComponentSpacePose;
+			ComponentSpacePose.InitPose(LocalSpacePose);
+
+			PoseHistory->EvaluateComponentSpace_AnyThread(DeltaTime, ComponentSpacePose, bStoreScales,
+				SharedData->RootBoneRecoveryTime, SharedData->RootBoneTranslationRecoveryRatio, SharedData->RootBoneRotationRecoveryRatio,
+				bNeedsReset, bCacheBones, BoneIndicesWithParents);
+		}
+	}
+}
+
+namespace UE::AnimNext
+{
+	AUTO_REGISTER_ANIM_TRAIT(FHistoryCollectorTrait)
+
+	// Trait implementation boilerplate
+	#define TRAIT_INTERFACE_ENUMERATOR(GeneratorMacro) \
+		GeneratorMacro(IUpdate) \
+		GeneratorMacro(IEvaluate) \
+		GeneratorMacro(IPoseHistory) \
+
+	// Trait implementation boilerplate
+	GENERATE_ANIM_TRAIT_IMPLEMENTATION(FHistoryCollectorTrait, TRAIT_INTERFACE_ENUMERATOR, NULL_ANIM_TRAIT_INTERFACE_ENUMERATOR, NULL_ANIM_TRAIT_EVENT_ENUMERATOR)
+	#undef TRAIT_INTERFACE_ENUMERATOR
+
+	void FHistoryCollectorTrait::PreUpdate(FUpdateTraversalContext& Context, const TTraitBinding<IUpdate>& Binding, const FTraitUpdateState& TraitState) const
+	{
+		const FSharedData* SharedData = Binding.GetSharedData<FSharedData>();
+		FInstanceData* InstanceData = Binding.GetInstanceData<FInstanceData>();
+		
+		InstanceData->DeltaTime = TraitState.GetDeltaTime();
+
+		// @todo: implement bResetOnBecomingRelevant
+		//bool bResetOnBecomingRelevant = SharedData->GetbResetOnBecomingRelevant(Binding);
+
+		InstanceData->PoseHistory.PreUpdate();
+
+		if (SharedData->GetbGenerateTrajectory(Binding))
+		{
+			// @todo: WIP trajectory generation
+			//FPoseSearchTrajectoryData::FSampling TrajectoryDataSampling;
+			//TrajectoryDataSampling.NumHistorySamples = FMath::Max(SharedData->PoseCount, SharedData->TrajectoryHistoryCount);
+			//TrajectoryDataSampling.SecondsPerHistorySample = SharedData->SamplingInterval;
+			//TrajectoryDataSampling.NumPredictionSamples = SharedData->TrajectoryPredictionCount;
+			//TrajectoryDataSampling.SecondsPerPredictionSample = SharedData->PredictionSamplingInterval;
+			//const bool bInitializeWithRefPose = SharedData->GetbInitializeWithRefPose(Binding);
+			//InstanceData->PoseHistory.GenerateTrajectory(InAnimInstance, InAnimInstance->GetDeltaSeconds(), InstanceData->TrajectoryData, TrajectoryDataSampling);
+		}
+		else
+		{
+			InstanceData->PoseHistory.SetTrajectory(SharedData->GetTrajectory(Binding));
+		}
+
+		IUpdate::PreUpdate(Context, Binding, TraitState);
+	}
+
+	void FHistoryCollectorTrait::PostEvaluate(FEvaluateTraversalContext& Context, const TTraitBinding<IEvaluate>& Binding) const
+	{
+		IEvaluate::PostEvaluate(Context, Binding);
+
+		const FSharedData* SharedData = Binding.GetSharedData<FSharedData>();
+		FInstanceData* InstanceData = Binding.GetInstanceData<FInstanceData>();
+		const bool bStoreScales = SharedData->GetbStoreScales(Binding);
+
+		Context.AppendTask(FAnimNextHistoryCollectorTask::Make(&InstanceData->PoseHistory, SharedData, bStoreScales, InstanceData->DeltaTime));
+	}
+
+	const UE::PoseSearch::IPoseHistory* FHistoryCollectorTrait::GetPoseHistory(FExecutionContext& Context, const TTraitBinding<IPoseHistory>& Binding) const
+	{
+		const FInstanceData* InstanceData = Binding.GetInstanceData<FInstanceData>();
+		return &InstanceData->PoseHistory;
+	}
+}
diff --git a/Engine/Plugins/Experimental/AnimNextPoseSearch/Source/AnimNextPoseSearch/Private/HistoryCollectorTrait.h b/Engine/Plugins/Experimental/AnimNextPoseSearch/Source/AnimNextPoseSearch/Private/HistoryCollectorTrait.h
new file mode 100644
--- /dev/null
+++ b/Engine/Plugins/Experimental/AnimNextPoseSearch/Source/AnimNextPoseSearch/Private/HistoryCollectorTrait.h
@@ -0,0 +1,152 @@
+// Copyright Epic Games, Inc. All Rights Reserved.
+
+#pragma once
+
+#include "CoreMinimal.h"
+#include "IPoseHistory.h"
+#include "PoseSearch/PoseSearchHistory.h"
+#include "PoseSearch/PoseSearchTrajectoryTypes.h"
+#include "TraitCore/Trait.h"
+#include "TraitInterfaces/IEvaluate.h"
+#include "TraitInterfaces/IUpdate.h"
+#include "HistoryCollectorTrait.generated.h"
+
+USTRUCT(meta = (DisplayName = "Pose History"))
+struct FAnimNextHistoryCollectorTraitSharedData : public FAnimNextTraitSharedData
+{
+	GENERATED_BODY()
+
+	// The maximum amount of poses that can be stored
+	UPROPERTY(EditAnywhere, Category = Settings, meta = (ClampMin="2"))
+	int32 PoseCount = 2;
+	
+	// how often in seconds poses are collected (if 0, it will collect every update)
+	UPROPERTY(EditAnywhere, Category = Settings, meta = (ClampMin="0"))
+	float SamplingInterval = 0.04f;
+
+	UPROPERTY(EditAnywhere, Category = Settings)
+	TArray<FBoneReference> CollectedBones;
+
+	// if true, the pose history will be initialized with a ref pose at the location and orientation of the AnimInstance.
+	UPROPERTY(EditAnywhere, Category = Settings)
+	bool bInitializeWithRefPose = false;
+
+	// Reset the pose history if it has become relevant to the graph after not being updated on previous frames.
+	UPROPERTY(EditAnywhere, Category = Settings)
+	bool bResetOnBecomingRelevant = true;
+
+	// if true pose scales will be cached, otherwise implied to be unitary scales
+	UPROPERTY(EditAnywhere, Category = Settings)
+	bool bStoreScales = false;
+
+	// time in seconds to recover to the reference skeleton root bone transform by RootBoneTranslationRecoveryRatio and RootBoneRotationRecoveryRatio
+	// from any eventual root bone modification. if zero the behaviour will be disabled (Experimental)
+	UPROPERTY(EditAnywhere, Category = Experimental, meta = (ClampMin="0"))
+	float RootBoneRecoveryTime = 0.f;
+
+	// ratio to recover to the reference skeleton root bone translation from any eventual root bone modification. zero for no recovery, 1 for full recovery
+	UPROPERTY(EditAnywhere, Category = Experimental, meta = (ClampMin="0", ClampMax="1", EditCondition = "RootBoneRecoveryTime > 0", EditConditionHides))
+	float RootBoneTranslationRecoveryRatio = 1.f;
+
+	// ratio to recover to the reference skeleton root bone rotation from any eventual root bone modification. zero for no recovery, 1 for full recovery
+	UPROPERTY(EditAnywhere, Category = Experimental, meta = (ClampMin="0", ClampMax="1", EditCondition = "RootBoneRecoveryTime > 0", EditConditionHides))
+	float RootBoneRotationRecoveryRatio = 1.f;
+
+//#if WITH_EDITORONLY_DATA
+//	UPROPERTY(EditAnywhere, Category = Debug)
+//	FLinearColor DebugColor = FLinearColor::Red;
+//#endif // WITH_EDITORONLY_DATA
+
+	// if true Trajectory the pose history node will generate the trajectory using the TrajectoryData parameters instead of relying on the input Trajectory (Experimental)
+	UPROPERTY(EditAnywhere, Category = Experimental)
+	bool bGenerateTrajectory = false;
+
+	// input Trajectory samples for pose search queries in Motion Matching. These are expected to be in the world space of the SkeletalMeshComponent.
+	// the trajectory sample with AccumulatedSeconds equals to zero (Trajectory.Samples[i].AccumulatedSeconds) is the sample of the previous frame of simulation (since MM works by matching the previous character pose)
+	UPROPERTY(EditAnywhere, Category = Settings, meta = (PinShownByDefault, EditCondition="!bGenerateTrajectory", EditConditionHides))
+	FPoseSearchQueryTrajectory Trajectory;
+
+	// Input Trajectory velocity will be multiplied by TrajectorySpeedMultiplier: values below 1 will result in selecting animation slower than requested from the original Trajectory
+	UPROPERTY(EditAnywhere, Category = Experimental, meta = (PinHiddenByDefault, ClampMin="0", EditCondition="!bGenerateTrajectory", EditConditionHides))
+	float TrajectorySpeedMultiplier = 1.f;
+
+	// if bGenerateTrajectory is true, this is the number of trajectory past (collected) samples
+	UPROPERTY(EditAnywhere, Category = Experimental, meta = (ClampMin = "2", EditCondition = "bGenerateTrajectory", EditConditionHides))
+	int32 TrajectoryHistoryCount = 10;
+
+	// if bGenerateTrajectory is true, this is the number of trajectory future (prediction) samples
+	UPROPERTY(EditAnywhere, Category = Experimental, meta = (ClampMin="2", EditCondition="bGenerateTrajectory", EditConditionHides))
+	int32 TrajectoryPredictionCount = 8;
+
+	// if bGenerateTrajectory is true, this is the sampling interval between trajectory future (prediction) samples
+	UPROPERTY(EditAnywhere, Category = Experimental, meta = (ClampMin="0.001", EditCondition="bGenerateTrajectory", EditConditionHides))
+	float PredictionSamplingInterval = 0.4f;
+
+	// if bGenerateTrajectory is true, TrajectoryData contains the tuning parameters to generate the trajectory
+	//UPROPERTY(EditAnywhere, Category = Experimental, meta=(EditCondition="bGenerateTrajectory", EditConditionHides))
+	//FPoseSearchTrajectoryData TrajectoryData;
+
+	// Latent pin support boilerplate
+	#define TRAIT_LATENT_PROPERTIES_ENUMERATOR(GeneratorMacro) \
+		GeneratorMacro(PoseCount) \
+		GeneratorMacro(SamplingInterval) \
+		GeneratorMacro(CollectedBones) \
+		GeneratorMacro(bInitializeWithRefPose) \
+		GeneratorMacro(bResetOnBecomingRelevant) \
+		GeneratorMacro(bStoreScales) \
+		GeneratorMacro(RootBoneRecoveryTime) \
+		GeneratorMacro(RootBoneTranslationRecoveryRatio) \
+		GeneratorMacro(RootBoneRotationRecoveryRatio) \
+		GeneratorMacro(bGenerateTrajectory) \
+		GeneratorMacro(Trajectory) \
+		GeneratorMacro(TrajectorySpeedMultiplier) \
+		GeneratorMacro(TrajectoryHistoryCount) \
+		GeneratorMacro(TrajectoryPredictionCount) \
+		GeneratorMacro(PredictionSamplingInterval) \
+
+	GENERATE_TRAIT_LATENT_PROPERTIES(FAnimNextHistoryCollectorTraitSharedData, TRAIT_LATENT_PROPERTIES_ENUMERATOR)
+	#undef TRAIT_LATENT_PROPERTIES_ENUMERATOR
+};
+
+USTRUCT()
+struct FAnimNextHistoryCollectorTask : public FAnimNextEvaluationTask
+{
+	GENERATED_BODY()
+
+	DECLARE_ANIM_EVALUATION_TASK(FAnimNextHistoryCollectorTask)
+
+	static FAnimNextHistoryCollectorTask Make(UE::PoseSearch::FPoseHistory* PoseHistory, const FAnimNextHistoryCollectorTraitSharedData* SharedData, bool bStoreScales, float DeltaTime);
+
+	// Task entry point
+	virtual void Execute(UE::AnimNext::FEvaluationVM& VM) const override;
+
+	UE::PoseSearch::FPoseHistory* PoseHistory = nullptr;
+	const FAnimNextHistoryCollectorTraitSharedData* SharedData = nullptr;
+	bool bStoreScales = false;
+	float DeltaTime = 0.f;
+};
+
+namespace UE::AnimNext
+{
+	struct FHistoryCollectorTrait : FAdditiveTrait, IUpdate, IEvaluate, IPoseHistory
+	{
+		DECLARE_ANIM_TRAIT(FHistoryCollectorTrait, 0xf1572b26, FAdditiveTrait)
+
+		using FSharedData = FAnimNextHistoryCollectorTraitSharedData;
+
+		struct FInstanceData : FTrait::FInstanceData
+		{
+			UE::PoseSearch::FPoseHistory PoseHistory;
+			float DeltaTime = 0.f;
+		};
+
+		// IUpdate impl
+		virtual void PreUpdate(FUpdateTraversalContext& Context, const TTraitBinding<IUpdate>& Binding, const FTraitUpdateState& TraitState) const override;
+
+		// IEvaluate impl
+		virtual void PostEvaluate(FEvaluateTraversalContext& Context, const TTraitBinding<IEvaluate>& Binding) const override;
+
+		// IPoseHistory impl
+		virtual const UE::PoseSearch::IPoseHistory* GetPoseHistory(FExecutionContext& Context, const TTraitBinding<IPoseHistory>& Binding) const override;
+	};
+}
diff --git a/Engine/Plugins/Experimental/AnimNextPoseSearch/Source/AnimNextPoseSearch/Private/IPoseHistory.cpp b/Engine/Plugins/Experimental/AnimNextPoseSearch/Source/AnimNextPoseSearch/Private/IPoseHistory.cpp
new file mode 100644
--- /dev/null
+++ b/Engine/Plugins/Experimental/AnimNextPoseSearch/Source/AnimNextPoseSearch/Private/IPoseHistory.cpp
@@ -0,0 +1,33 @@
+// Copyright Epic Games, Inc. All Rights Reserved.
+
+#include "IPoseHistory.h"
+#include "TraitCore/ExecutionContext.h"
+
+namespace UE::AnimNext
+{
+	AUTO_REGISTER_ANIM_TRAIT_INTERFACE(IPoseHistory)
+
+#if WITH_EDITOR
+	const FText& IPoseHistory::GetDisplayName() const
+	{
+		static FText InterfaceName = NSLOCTEXT("TraitInterfaces", "TraitInterface_IPoseHistory_Name", "Pose History");
+		return InterfaceName;
+	}
+	const FText& IPoseHistory::GetDisplayShortName() const
+	{
+		static FText InterfaceShortName = NSLOCTEXT("TraitInterfaces", "TraitInterface_IPoseHistory_ShortName", "PH");
+		return InterfaceShortName;
+	}
+#endif // WITH_EDITOR
+
+	const UE::PoseSearch::IPoseHistory* IPoseHistory::GetPoseHistory(FExecutionContext& Context, const TTraitBinding<IPoseHistory>& Binding) const
+	{
+		TTraitBinding<IPoseHistory> SuperBinding;
+		if (Binding.GetStackInterfaceSuper(SuperBinding))
+		{
+			return SuperBinding.GetPoseHistory(Context);
+		}
+
+		return nullptr;
+	}
+}
diff --git a/Engine/Plugins/Experimental/AnimNextPoseSearch/Source/AnimNextPoseSearch/Private/IPoseHistory.h b/Engine/Plugins/Experimental/AnimNextPoseSearch/Source/AnimNextPoseSearch/Private/IPoseHistory.h
new file mode 100644
--- /dev/null
+++ b/Engine/Plugins/Experimental/AnimNextPoseSearch/Source/AnimNextPoseSearch/Private/IPoseHistory.h
@@ -0,0 +1,45 @@
+// Copyright Epic Games, Inc. All Rights Reserved.
+
+#pragma once
+
+#include "CoreMinimal.h"
+#include "TraitCore/ITraitInterface.h"
+#include "TraitCore/TraitBinding.h"
+
+namespace UE::PoseSearch
+{
+	struct IPoseHistory;
+}
+
+namespace UE::AnimNext
+{
+	/**
+	 * IPoseHistory
+	 */
+	struct IPoseHistory : ITraitInterface
+	{
+		DECLARE_ANIM_TRAIT_INTERFACE(IPoseHistory, 0x1c2cffc4)
+
+		virtual const UE::PoseSearch::IPoseHistory* GetPoseHistory(FExecutionContext& Context, const TTraitBinding<IPoseHistory>& Binding) const;
+
+#if WITH_EDITOR
+		virtual const FText& GetDisplayName() const override;
+		virtual const FText& GetDisplayShortName() const override;
+#endif // WITH_EDITOR
+	};
+
+	/**
+	 * Specialization for trait binding.
+	 */
+	template<>
+	struct TTraitBinding<IPoseHistory> : FTraitBinding
+	{
+		const UE::PoseSearch::IPoseHistory* GetPoseHistory(FExecutionContext& Context) const
+		{
+			return GetInterface()->GetPoseHistory(Context, *this);
+		}
+
+	protected:
+		const IPoseHistory* GetInterface() const { return GetInterfaceTyped<IPoseHistory>(); }
+	};
+}
diff --git a/Engine/Plugins/Experimental/AnimNextPoseSearch/Source/AnimNextPoseSearch/Private/Module.cpp b/Engine/Plugins/Experimental/AnimNextPoseSearch/Source/AnimNextPoseSearch/Private/Module.cpp
new file mode 100644
--- /dev/null
+++ b/Engine/Plugins/Experimental/AnimNextPoseSearch/Source/AnimNextPoseSearch/Private/Module.cpp
@@ -0,0 +1,47 @@
+// Copyright Epic Games, Inc. All Rights Reserved.
+
+#include "CoreMinimal.h"
+#include "Module/AnimNextModule.h"
+#include "Modules/ModuleInterface.h"
+#include "Modules/ModuleManager.h"
+#include "RigVMCore/RigVMRegistry.h"
+
+#include "IAnimNextModuleInterface.h"
+#include "EvaluationVM/EvaluationVM.h"
+#include "PoseSearch/PoseSearchDatabase.h"
+#include "PoseSearch/PoseSearchSchema.h"
+#include "PoseSearch/PoseSearchTrajectoryTypes.h"
+
+namespace UE::AnimNext::PoseSearch
+{
+
+class FModule : public IModuleInterface
+{
+public:
+	virtual void StartupModule() override
+	{
+		static TPair<UClass*, FRigVMRegistry::ERegisterObjectOperation> const AllowedObjectTypes[] =
+		{
+			{ UPoseSearchDatabase::StaticClass(), FRigVMRegistry::ERegisterObjectOperation::Class },
+			{ UPoseSearchSchema::StaticClass(), FRigVMRegistry::ERegisterObjectOperation::Class }
+		};
+
+		static UScriptStruct* const AllowedStructTypes[] =
+		{
+			FPoseSearchQueryTrajectorySample::StaticStruct(),
+			FPoseSearchQueryTrajectory::StaticStruct()
+		};
+
+		FRigVMRegistry& RigVMRegistry = FRigVMRegistry::Get();
+		RigVMRegistry.RegisterObjectTypes(AllowedObjectTypes);
+		RigVMRegistry.RegisterStructTypes(AllowedStructTypes);
+	}
+
+	virtual void ShutdownModule() override
+	{
+	}
+};
+
+}
+
+IMPLEMENT_MODULE(UE::AnimNext::PoseSearch::FModule, AnimNextPoseSearch)
diff --git a/Engine/Plugins/Experimental/AnimNextPoseSearch/Source/AnimNextPoseSearch/Private/MotionMatchingTrait.cpp b/Engine/Plugins/Experimental/AnimNextPoseSearch/Source/AnimNextPoseSearch/Private/MotionMatchingTrait.cpp
new file mode 100644
--- /dev/null
+++ b/Engine/Plugins/Experimental/AnimNextPoseSearch/Source/AnimNextPoseSearch/Private/MotionMatchingTrait.cpp
@@ -0,0 +1,129 @@
+// Copyright Epic Games, Inc. All Rights Reserved.
+
+#include "MotionMatchingTrait.h"
+#include "Animation/AnimSequence.h"
+#include "DataInterface/DataInterfaceStructAdapter.h"
+#include "IPoseHistory.h"
+#include "PlayAnim/PlayAnimRequest.h"
+#include "PoseSearch/PoseSearchDatabase.h"
+#include "TraitInterfaces/IBlendStack.h"
+
+namespace UE::AnimNext
+{
+	AUTO_REGISTER_ANIM_TRAIT(FMotionMatchingTrait)
+
+	#define TRAIT_INTERFACE_ENUMERATOR(GeneratorMacro) \
+		GeneratorMacro(IUpdate) \
+
+	// Trait implementation boilerplate
+	GENERATE_ANIM_TRAIT_IMPLEMENTATION(FMotionMatchingTrait, TRAIT_INTERFACE_ENUMERATOR, NULL_ANIM_TRAIT_INTERFACE_ENUMERATOR, NULL_ANIM_TRAIT_EVENT_ENUMERATOR)
+	#undef TRAIT_INTERFACE_ENUMERATOR
+
+	void FMotionMatchingTrait::PreUpdate(FUpdateTraversalContext& Context, const TTraitBinding<IUpdate>& Binding, const FTraitUpdateState& TraitState) const
+	{
+		using namespace UE::PoseSearch;
+
+		IUpdate::PreUpdate(Context, Binding, TraitState);
+
+		TTraitBinding<IBlendStack> BlendStackBinding;
+		if (!Binding.GetStackInterface(BlendStackBinding))
+		{
+			// @todo: LOG badness
+			return;
+		}
+
+		TTraitBinding<IPoseHistory> PoseHistoryTrait;
+		if (!Binding.GetStackInterface(PoseHistoryTrait))
+		{
+			// @todo: LOG badness
+			return;
+		}
+
+		const UE::PoseSearch::IPoseHistory* PoseHistory = PoseHistoryTrait.GetPoseHistory(Context);
+		if (!PoseHistory)
+		{
+			// @todo: LOG badness
+			return;
+		}
+
+		const FSharedData* SharedData = Binding.GetSharedData<FSharedData>();
+		check(SharedData);
+
+		//FInstanceData* InstanceData = Binding.GetInstanceData<FInstanceData>();
+		//check(InstanceData);
+
+		TArray<const UObject*> AssetsToSearch;
+		if (const UObject* DataBaseToSearch = SharedData->GetDatabase(Binding))
+		{
+			AssetsToSearch.Add(DataBaseToSearch);
+		}
+
+		if (AssetsToSearch.IsEmpty())
+		{
+			// @todo: LOG badness
+			return;
+		}
+
+		IBlendStack::FGraphRequestPtr ActiveGraphRequest;
+		BlendStackBinding.GetActiveGraph(Context, ActiveGraphRequest);
+
+		FPoseSearchContinuingProperties ContinuingProperties;
+		if (ActiveGraphRequest)
+		{
+			check(ActiveGraphRequest->AnimationGraph);
+
+			for (const FInstancedStruct& Payload : ActiveGraphRequest->Payloads)
+			{
+				if (const FAnimNextPlayAnimPayload* PlayAnimPayload = Payload.GetPtr<FAnimNextPlayAnimPayload>())
+				{
+					ContinuingProperties.PlayingAsset = PlayAnimPayload->AnimationObject;
+
+					// it'd be nice to be able to have the value from the boud payload instead of having to ask to the ITimeline
+					// ContinuingProperties.PlayingAssetAccumulatedTime = PlayAnimPayload->StartPosition;
+					
+					TTraitBinding<ITimeline> TimelineTrait;
+					if (Binding.GetStackInterface(TimelineTrait))
+					{
+						const FTimelineProgress TimelineProgress = TimelineTrait.GetProgress(Context);
+						ContinuingProperties.PlayingAssetAccumulatedTime = TimelineProgress.GetPosition();
+					}
+
+					break;
+				}
+			}
+		}
+
+		TArray<UAnimInstance*> AnimInstances; AnimInstances.Add(nullptr);
+		const TArrayView<const FRole> Roles = MakeArrayView(&DefaultRole, 1);
+		const TArrayView<const UE::PoseSearch::IPoseHistory*> PoseHistories = MakeArrayView(&PoseHistory, 1);
+		const FPoseSearchFutureProperties FutureProperties;
+
+		const FSearchResult SearchResult = UPoseSearchLibrary::MotionMatch(AnimInstances, Roles, PoseHistories,  AssetsToSearch, ContinuingProperties, FutureProperties);
+
+		if (!SearchResult.bIsContinuingPoseSearch)
+		{
+			if (const UE::PoseSearch::FSearchIndexAsset* SearchIndexAsset = SearchResult.GetSearchIndexAsset())
+			{
+				if (const FPoseSearchDatabaseAnimationAssetBase* DatabaseAnimationAssetBase = SearchResult.Database->GetDatabaseAnimationAsset<FPoseSearchDatabaseAnimationAssetBase>(*SearchIndexAsset))
+				{
+					if (UAnimSequence* AnimSequence = Cast<UAnimSequence>(DatabaseAnimationAssetBase->GetAnimationAsset()))
+					{
+						UE::AnimNext::IBlendStack::FGraphRequest NewGraphRequest;
+						NewGraphRequest.BlendArgs = SharedData->BlendArguments;
+						NewGraphRequest.AnimationGraph = SharedData->AnimationGraph;
+
+						FInstancedStruct& NewPayload = NewGraphRequest.Payloads.AddDefaulted_GetRef();
+						NewPayload.InitializeAs<FAnimNextPlayAnimPayload>();
+						FAnimNextPlayAnimPayload& NewPlayAnimPayload = NewPayload.GetMutable<FAnimNextPlayAnimPayload>();
+						NewPlayAnimPayload.AnimationObject = AnimSequence;
+						NewPlayAnimPayload.PlayRate = 1.f;
+						NewPlayAnimPayload.StartPosition = SearchResult.AssetTime;
+						NewPlayAnimPayload.IsLooping = SearchIndexAsset->IsLooping();
+
+						BlendStackBinding.PushGraph(Context, MoveTemp(NewGraphRequest));
+					}
+				}
+			}
+		}
+	}
+}
diff --git a/Engine/Plugins/Experimental/AnimNextPoseSearch/Source/AnimNextPoseSearch/Private/MotionMatchingTrait.h b/Engine/Plugins/Experimental/AnimNextPoseSearch/Source/AnimNextPoseSearch/Private/MotionMatchingTrait.h
new file mode 100644
--- /dev/null
+++ b/Engine/Plugins/Experimental/AnimNextPoseSearch/Source/AnimNextPoseSearch/Private/MotionMatchingTrait.h
@@ -0,0 +1,61 @@
+// Copyright Epic Games, Inc. All Rights Reserved.
+
+#pragma once
+
+#include "CoreMinimal.h"
+#include "PoseSearch/PoseSearchLibrary.h"
+#include "TraitCore/Trait.h"
+#include "TraitCore/TraitSharedData.h"
+#include "TraitInterfaces/IUpdate.h"
+#include "MotionMatchingTrait.generated.h"
+
+USTRUCT(meta = (DisplayName = "Motion Matching"))
+struct FMotionMatchingTraitSharedData : public FAnimNextTraitSharedData
+{
+	GENERATED_BODY()
+
+	// Asset to use as a sub-graph
+	UPROPERTY(EditAnywhere, Category = "Default")
+	TObjectPtr<const UAnimNextAnimationGraph> AnimationGraph;
+
+	// Entry point in the Subgraph that we will use
+	UPROPERTY(EditAnywhere, Category = "Default", meta=(CustomWidget = "ParamName", AllowedParamType = "FAnimNextEntryPoint"))
+	FName EntryPoint;
+
+	// @todo: support array of databases
+	// database to search from
+	UPROPERTY(EditAnywhere, Category = "Default")
+	TObjectPtr<const UObject> Database;
+
+	UPROPERTY(EditAnywhere, Category = "Default")
+	FAlphaBlendArgs BlendArguments;
+
+	// Latent pin support boilerplate
+	#define TRAIT_LATENT_PROPERTIES_ENUMERATOR(GeneratorMacro) \
+		GeneratorMacro(AnimationGraph) \
+		GeneratorMacro(EntryPoint) \
+		GeneratorMacro(Database) \
+		GeneratorMacro(BlendArguments) \
+
+	GENERATE_TRAIT_LATENT_PROPERTIES(FMotionMatchingTraitSharedData, TRAIT_LATENT_PROPERTIES_ENUMERATOR)
+	#undef TRAIT_LATENT_PROPERTIES_ENUMERATOR
+};
+
+namespace UE::AnimNext
+{
+
+struct FMotionMatchingTrait : FAdditiveTrait, IUpdate
+{
+	DECLARE_ANIM_TRAIT(FMotionMatchingTrait, 0xf1572ff4, FAdditiveTrait)
+
+	using FSharedData = FMotionMatchingTraitSharedData;
+		
+	struct FInstanceData : FTrait::FInstanceData
+	{
+	};
+
+	// IUpdate impl
+	virtual void PreUpdate(FUpdateTraversalContext& Context, const TTraitBinding<IUpdate>& Binding, const FTraitUpdateState& TraitState) const override;
+};
+
+} // namespace UE::AnimNext
