AnimNextPoseSearch - fix for IsBetaVersion

diff --git a/Engine/Plugins/Experimental/AnimNextPoseSearch/AnimNextPoseSearch.uplugin b/Engine/Plugins/Experimental/AnimNextPoseSearch/AnimNextPoseSearch.uplugin
--- a/Engine/Plugins/Experimental/AnimNextPoseSearch/AnimNextPoseSearch.uplugin
+++ b/Engine/Plugins/Experimental/AnimNextPoseSearch/AnimNextPoseSearch.uplugin
@@ -12,7 +12,7 @@
 	"SupportURL" : "",
 	"EnabledByDefault" : false,
 	"CanContainContent" : false,
-	"IsBetaVersion": true,
+	"IsBetaVersion": false,
 	"IsExperimentalVersion": true,
 	"Installed" : false,
 	"Plugins":
