AnimNextPoseSearch - using UpdateMotionMatchingState in MM trait instead of the MotionMatch function

diff --git a/Engine/Plugins/Animation/PoseSearch/Source/Runtime/Private/AnimNode_MotionMatching.cpp b/Engine/Plugins/Animation/PoseSearch/Source/Runtime/Private/AnimNode_MotionMatching.cpp
--- a/Engine/Plugins/Animation/PoseSearch/Source/Runtime/Private/AnimNode_MotionMatching.cpp
+++ b/Engine/Plugins/Animation/PoseSearch/Source/Runtime/Private/AnimNode_MotionMatching.cpp
@@ -38,7 +38,9 @@ void FAnimNode_MotionMatching::Initialize_AnyThread(const FAnimationInitializeCo
 	GetEvaluateGraphExposedInputs().Execute(Context);
 
 	FAnimNode_BlendStack_Standalone::Initialize_AnyThread(Context);
+	PRAGMA_DISABLE_DEPRECATION_WARNINGS
 	MotionMatchingState.Reset(Context.AnimInstanceProxy->GetComponentTransform());
+	PRAGMA_ENABLE_DEPRECATION_WARNINGS
 }
 
 void FAnimNode_MotionMatching::Evaluate_AnyThread(FPoseContext& Output)
@@ -96,6 +98,8 @@ void FAnimNode_MotionMatching::UpdateAssetPlayer(const FAnimationUpdateContext&
 
 	using namespace UE::PoseSearch;
 
+	check(Context.AnimInstanceProxy);
+
 	GetEvaluateGraphExposedInputs().Execute(Context);
 
 	bool bNeedsReset =
@@ -118,7 +122,9 @@ void FAnimNode_MotionMatching::UpdateAssetPlayer(const FAnimationUpdateContext&
 	// If we just became relevant and haven't been initialized yet, then reset motion matching state, otherwise update the asset time using the player node.
 	if (bNeedsReset)
 	{
+		PRAGMA_DISABLE_DEPRECATION_WARNINGS
 		MotionMatchingState.Reset(Context.AnimInstanceProxy->GetComponentTransform());
+		PRAGMA_ENABLE_DEPRECATION_WARNINGS
 		FAnimNode_BlendStack_Standalone::Reset();
 	}
 	else
@@ -163,25 +169,40 @@ void FAnimNode_MotionMatching::UpdateAssetPlayer(const FAnimationUpdateContext&
 	}
 #endif // ENABLE_ANIM_DEBUG
 
-	// Execute core motion matching algorithm
-	UPoseSearchLibrary::UpdateMotionMatchingState(
-		Context,
-		DatabasesToSearch,
-		BlendTime,
-		MaxActiveBlends,
-		PoseJumpThresholdTime,
-		PoseReselectHistory,
-		SearchThrottleTime,
-		PlayRate,
-		MotionMatchingState,
-		NextUpdateInterruptMode,
-		bShouldSearch,
-		bShouldUseCachedChannelData
-		#if ENABLE_ANIM_DEBUG
-		, CVarAnimNodeMotionMatchingDrawQuery.GetValueOnAnyThread()
-		, CVarAnimNodeMotionMatchingDrawCurResult.GetValueOnAnyThread()
-		#endif // ENABLE_ANIM_DEBUG
-	);
+	if (DatabasesToSearch.IsEmpty())
+	{
+		Context.LogMessage(
+			EMessageSeverity::Error,
+			LOCTEXT("NoDatabases", "No database assets provided for motion matching."));
+	}
+	else
+	{
+		const float DeltaTime = Context.GetDeltaTime();
+
+		const IPoseHistory* PoseHistory = nullptr;
+		if (const FPoseHistoryProvider* PoseHistoryProvider = Context.GetMessage<FPoseHistoryProvider>())
+		{
+			PoseHistory = &PoseHistoryProvider->GetPoseHistory();
+		}
+
+		UPoseSearchLibrary::UpdateMotionMatchingState(
+			Context.AnimInstanceProxy->GetAnimInstanceObject(),
+			PoseHistory,
+			DatabasesToSearch,
+			Context.GetDeltaTime(),
+			PoseJumpThresholdTime,
+			PoseReselectHistory,
+			bShouldSearch ? SearchThrottleTime : UE_BIG_NUMBER,
+			PlayRate,
+			MotionMatchingState,
+			NextUpdateInterruptMode,
+			bShouldUseCachedChannelData
+#if ENABLE_ANIM_DEBUG
+			, CVarAnimNodeMotionMatchingDrawQuery.GetValueOnAnyThread()
+			, CVarAnimNodeMotionMatchingDrawCurResult.GetValueOnAnyThread()
+#endif // ENABLE_ANIM_DEBUG
+		);
+	}
 
 	const bool bJumpToPose = MotionMatchingState.CurrentSearchResult.IsValid() && !MotionMatchingState.CurrentSearchResult.bIsContinuingPoseSearch;
 
diff --git a/Engine/Plugins/Animation/PoseSearch/Source/Runtime/Private/PoseSearchLibrary.cpp b/Engine/Plugins/Animation/PoseSearch/Source/Runtime/Private/PoseSearchLibrary.cpp
--- a/Engine/Plugins/Animation/PoseSearch/Source/Runtime/Private/PoseSearchLibrary.cpp
+++ b/Engine/Plugins/Animation/PoseSearch/Source/Runtime/Private/PoseSearchLibrary.cpp
@@ -264,11 +264,6 @@ namespace UE::PoseSearch
 
 void FMotionMatchingState::Reset(const FTransform& ComponentTransform)
 {
-	CurrentSearchResult.Reset();
-	// Set the elapsed time to INFINITY to trigger a search right away
-	ElapsedPoseSearchTime = std::numeric_limits<float>::infinity();
-	WantedPlayRate = 1.f;
-	
 PRAGMA_DISABLE_DEPRECATION_WARNINGS
 	bJumpedToPose = false;
 	ComponentDeltaYaw = 0.f;
@@ -276,22 +271,21 @@ PRAGMA_DISABLE_DEPRECATION_WARNINGS
 	AnimationDeltaYaw = 0.f;
 PRAGMA_ENABLE_DEPRECATION_WARNINGS
 
-	PoseIndicesHistory.Reset();
+	Reset();
 }
 
-void FMotionMatchingState::AdjustAssetTime(float AssetTime)
+void FMotionMatchingState::Reset()
 {
-	CurrentSearchResult.Update(AssetTime);
+	CurrentSearchResult.Reset();
+	// Set the elapsed time to INFINITY to trigger a search right away
+	ElapsedPoseSearchTime = std::numeric_limits<float>::infinity();
+	WantedPlayRate = 1.f;
+	PoseIndicesHistory.Reset();
 }
 
-void FMotionMatchingState::JumpToPose(const FAnimationUpdateContext& Context, const UE::PoseSearch::FSearchResult& Result, int32 MaxActiveBlends, float BlendTime)
+void FMotionMatchingState::AdjustAssetTime(float AssetTime)
 {
-	// Remember which pose and sequence we're playing from the database
-	CurrentSearchResult = Result;
-
-PRAGMA_DISABLE_DEPRECATION_WARNINGS
-	bJumpedToPose = true;
-PRAGMA_ENABLE_DEPRECATION_WARNINGS
+	CurrentSearchResult.Update(AssetTime);
 }
 
 FVector FMotionMatchingState::GetEstimatedFutureRootMotionVelocity() const
@@ -569,12 +563,8 @@ void UPoseSearchLibrary::UpdateMotionMatchingState(
 	bool bDebugDrawQuery,
 	bool bDebugDrawCurResult)
 {
-	QUICK_SCOPE_CYCLE_COUNTER(STAT_PoseSearch_Update);
-
 	using namespace UE::PoseSearch;
 
-	check(Context.AnimInstanceProxy);
-
 	if (Databases.IsEmpty())
 	{
 		Context.LogMessage(
@@ -582,17 +572,45 @@ void UPoseSearchLibrary::UpdateMotionMatchingState(
 			LOCTEXT("NoDatabases", "No database assets provided for motion matching."));
 		return;
 	}
-	
-	const float DeltaTime = Context.GetDeltaTime();
+
 	const IPoseHistory* PoseHistory = nullptr;
-	if (FPoseHistoryProvider* PoseHistoryProvider = Context.GetMessage<FPoseHistoryProvider>())
+	if (const FPoseHistoryProvider* PoseHistoryProvider = Context.GetMessage<FPoseHistoryProvider>())
 	{
 		PoseHistory = &PoseHistoryProvider->GetPoseHistory();
 	}
 
+	check(Context.AnimInstanceProxy);
+	UpdateMotionMatchingState(Context.AnimInstanceProxy->GetAnimInstanceObject(), PoseHistory, Databases, Context.GetDeltaTime(),
+		PoseJumpThresholdTime, PoseReselectHistory, bShouldSearch ? SearchThrottleTime : UE_BIG_NUMBER, PlayRate, InOutMotionMatchingState,
+		InterruptMode, bShouldUseCachedChannelData, bDebugDrawQuery, bDebugDrawCurResult);
+}
+
+void UPoseSearchLibrary::UpdateMotionMatchingState(
+	const UObject* AnimContext,
+	const UE::PoseSearch::IPoseHistory* PoseHistory,
+	const TArray<TObjectPtr<const UPoseSearchDatabase>>& Databases,
+	float DeltaTime,
+	const FFloatInterval& PoseJumpThresholdTime,
+	float PoseReselectHistory,
+	float SearchThrottleTime,
+	const FFloatInterval& PlayRate,
+	FMotionMatchingState& InOutMotionMatchingState,
+	EPoseSearchInterruptMode InterruptMode,
+	bool bShouldUseCachedChannelData,
+	bool bDebugDrawQuery,
+	bool bDebugDrawCurResult)
+{
+	QUICK_SCOPE_CYCLE_COUNTER(STAT_PoseSearch_Update);
+
+	using namespace UE::PoseSearch;
+
+	if (Databases.IsEmpty())
+	{
+		UE_LOG(LogPoseSearch, Error, TEXT("No database assets provided for motion matching"));
+		return;
+	}
+
 	FMemMark Mark(FMemStack::Get());
-	const UAnimInstance* AnimInstance = Cast<const UAnimInstance>(Context.AnimInstanceProxy->GetAnimInstanceObject());
-	check(AnimInstance);
 
 	const UPoseSearchDatabase* CurrentResultDatabase = InOutMotionMatchingState.CurrentSearchResult.Database.Get();
 	if (IsInvalidatingContinuingPose(InterruptMode, CurrentResultDatabase, Databases))
@@ -601,12 +619,12 @@ void UPoseSearchLibrary::UpdateMotionMatchingState(
 	}
 
 	FSearchContext SearchContext(0.f, &InOutMotionMatchingState.PoseIndicesHistory, InOutMotionMatchingState.CurrentSearchResult, PoseJumpThresholdTime);
-	SearchContext.AddRole(DefaultRole, AnimInstance, PoseHistory);
+	SearchContext.AddRole(DefaultRole, AnimContext, PoseHistory);
 
 	const bool bCanAdvance = InOutMotionMatchingState.CurrentSearchResult.CanAdvance(DeltaTime);
 
 	// If we can't advance or enough time has elapsed since the last pose jump then search
-	const bool bSearch = !bCanAdvance || (bShouldSearch && (InOutMotionMatchingState.ElapsedPoseSearchTime >= SearchThrottleTime));
+	const bool bSearch = !bCanAdvance || (InOutMotionMatchingState.ElapsedPoseSearchTime >= SearchThrottleTime);
 	if (bSearch)
 	{
 		InOutMotionMatchingState.ElapsedPoseSearchTime = 0.f;
diff --git a/Engine/Plugins/Animation/PoseSearch/Source/Runtime/Public/PoseSearch/PoseSearchDefines.h b/Engine/Plugins/Animation/PoseSearch/Source/Runtime/Public/PoseSearch/PoseSearchDefines.h
--- a/Engine/Plugins/Animation/PoseSearch/Source/Runtime/Public/PoseSearch/PoseSearchDefines.h
+++ b/Engine/Plugins/Animation/PoseSearch/Source/Runtime/Public/PoseSearch/PoseSearchDefines.h
@@ -19,7 +19,7 @@ constexpr EParallelForFlags ParallelForFlags = EParallelForFlags::ForceSingleThr
 constexpr EParallelForFlags ParallelForFlags = EParallelForFlags::None;
 #endif // UE_POSE_SEARCH_FORCE_SINGLE_THREAD
 
-DECLARE_LOG_CATEGORY_EXTERN(LogPoseSearch, Log, All);
+POSESEARCH_API DECLARE_LOG_CATEGORY_EXTERN(LogPoseSearch, Log, All);
 
 namespace UE::PoseSearch
 {
diff --git a/Engine/Plugins/Animation/PoseSearch/Source/Runtime/Public/PoseSearch/PoseSearchLibrary.h b/Engine/Plugins/Animation/PoseSearch/Source/Runtime/Public/PoseSearch/PoseSearchLibrary.h
--- a/Engine/Plugins/Animation/PoseSearch/Source/Runtime/Public/PoseSearch/PoseSearchLibrary.h
+++ b/Engine/Plugins/Animation/PoseSearch/Source/Runtime/Public/PoseSearch/PoseSearchLibrary.h
@@ -47,16 +47,16 @@ struct FMotionMatchingState
 {
 PRAGMA_ENABLE_DEPRECATION_WARNINGS
 	// Reset the state to a default state using the current Database
+	UE_DEPRECATED(5.6, "Use other Reset signatures instead")
 	void Reset(const FTransform& ComponentTransform);
 
+	// Reset the state to a default state using the current Database
+	POSESEARCH_API void Reset();
+
 	// Attempts to set the internal state to match the provided asset time including updating the internal DbPoseIdx. 
 	// If the provided asset time is out of bounds for the currently playing asset then this function will reset the 
 	// state back to the default state.
-	void AdjustAssetTime(float AssetTime);
-
-	// Internally stores the 'jump' to a new pose/sequence index and asset time for evaluation
-	UE_DEPRECATED(5.6, "Set the CurrentSearchResult directly if needed, since bJumpedToPose has already been deprecated")
-	void JumpToPose(const FAnimationUpdateContext& Context, const UE::PoseSearch::FSearchResult& Result, int32 MaxActiveBlends, float BlendTime);
+	POSESEARCH_API void AdjustAssetTime(float AssetTime);
 
 	void UpdateWantedPlayRate(const UE::PoseSearch::FSearchContext& SearchContext, const FFloatInterval& PlayRate, float TrajectorySpeedMultiplier);
 
@@ -190,20 +190,36 @@ public:
 	/**
 	* Implementation of the core motion matching algorithm
 	*
-	* @param Context						Input animation update context providing access to the proxy and delta time
+	* @param AnimContext					Input animation context, either UAnimInstance or UAnimNextComponent
+	* @param PoseHistory					Input IPoseHistory, used to gather historical information about trajectory and bone transforms
 	* @param Databases						Input array of databases to search
-	* @param BlendTime						Input time in seconds to blend out to the new pose. Uses either inertial blending, requiring an Inertialization node after this node, or the internal blend stack, if MaxActiveBlends is greater than zero.
-	* @param MaxActiveBlends				Input number of max active animation segments being blended together in the blend stack. If MaxActiveBlends is zero then the blend stack is disabled.
+	* @param DeltaTime						Input DeltaTime
 	* @param PoseJumpThresholdTime			Input don't jump to poses of the same segment that are within the interval this many seconds away from the continuing pose.
 	* @param PoseReselectHistory			Input prevent re-selection of poses that have been selected previously within this much time (in seconds) in the past. This is across all animation segments that have been selected within this time range.
 	* @param SearchThrottleTime				Input minimum amount of time to wait between searching for a new pose segment. It allows users to define how often the system searches, default for locomotion is searching every update, but you may only want to search once for other situations, like jump.
 	* @param PlayRate						Input effective range of play rate that can be applied to the animations to account for discrepancies in estimated velocity between the movement modeland the animation.
 	* @param InOutMotionMatchingState		Input/Output encapsulated motion matching algorithm and state
 	* @param InterruptMode					Input continuing pose search interrupt mode
-	* @param bShouldSearch					Input if false search will happen only if there's no valid continuing pose
+	* @param bShouldUseCachedChannelData	Input if true, motion matching will try to reuse the continuing pose channels features across multiple schemas without querying the pose history
 	* @param bDebugDrawQuery				Input draw the composed query if valid
 	* @param bDebugDrawCurResult			Input draw the current result if valid
 	*/
+	static void UpdateMotionMatchingState(
+		const UObject* AnimContext,
+		const UE::PoseSearch::IPoseHistory* PoseHistory,
+		const TArray<TObjectPtr<const UPoseSearchDatabase>>& Databases,
+		float DeltaTime,
+		const FFloatInterval& PoseJumpThresholdTime,
+		float PoseReselectHistory,
+		float SearchThrottleTime,
+		const FFloatInterval& PlayRate,
+		FMotionMatchingState& InOutMotionMatchingState,
+		EPoseSearchInterruptMode InterruptMode,
+		bool bShouldUseCachedChannelData,
+		bool bDebugDrawQuery = false,
+		bool bDebugDrawCurResult = false);
+
+	UE_DEPRECATED(5.6, "Use other UpdateMotionMatchingState signatures instead")
 	static void UpdateMotionMatchingState(
 		const FAnimationUpdateContext& Context,
 		const TArray<TObjectPtr<const UPoseSearchDatabase>>& Databases,
diff --git a/Engine/Plugins/Experimental/AnimNextPoseSearch/Source/AnimNextPoseSearch/Private/MotionMatchingTrait.cpp b/Engine/Plugins/Experimental/AnimNextPoseSearch/Source/AnimNextPoseSearch/Private/MotionMatchingTrait.cpp
--- a/Engine/Plugins/Experimental/AnimNextPoseSearch/Source/AnimNextPoseSearch/Private/MotionMatchingTrait.cpp
+++ b/Engine/Plugins/Experimental/AnimNextPoseSearch/Source/AnimNextPoseSearch/Private/MotionMatchingTrait.cpp
@@ -6,6 +6,7 @@
 #include "IPoseHistory.h"
 #include "PlayAnim/PlayAnimRequest.h"
 #include "PoseSearch/PoseSearchDatabase.h"
+#include "PoseSearch/PoseSearchDerivedData.h"
 #include "TraitInterfaces/IBlendStack.h"
 
 namespace UE::AnimNext
@@ -28,88 +29,133 @@ namespace UE::AnimNext
 		TTraitBinding<IBlendStack> BlendStackBinding;
 		if (!Binding.GetStackInterface(BlendStackBinding))
 		{
-			// @todo: LOG badness
+			UE_LOG(LogPoseSearch, Error, TEXT("FMotionMatchingTrait::PreUpdate, missing IBlendStack"));
 			return;
 		}
 
 		TTraitBinding<IPoseHistory> PoseHistoryTrait;
 		if (!Binding.GetStackInterface(PoseHistoryTrait))
 		{
-			// @todo: LOG badness
+			UE_LOG(LogPoseSearch, Error, TEXT("FMotionMatchingTrait::PreUpdate, missing IPoseHistory"));
+			return;
+		}
+
+		TTraitBinding<ITimeline> TimelineTrait;
+		if (!Binding.GetStackInterface(TimelineTrait))
+		{
+			UE_LOG(LogPoseSearch, Error, TEXT("FMotionMatchingTrait::PreUpdate, missing ITimeline"));
 			return;
 		}
 
 		const UE::PoseSearch::IPoseHistory* PoseHistory = PoseHistoryTrait.GetPoseHistory(Context);
 		if (!PoseHistory)
 		{
-			// @todo: LOG badness
+			UE_LOG(LogPoseSearch, Error, TEXT("FMotionMatchingTrait::PreUpdate, missing IPoseHistory"));
 			return;
 		}
 
 		const FSharedData* SharedData = Binding.GetSharedData<FSharedData>();
 		check(SharedData);
 
-		//FInstanceData* InstanceData = Binding.GetInstanceData<FInstanceData>();
-		//check(InstanceData);
+		FInstanceData* InstanceData = Binding.GetInstanceData<FInstanceData>();
+		check(InstanceData);
+		FMotionMatchingState& MotionMatchingState = InstanceData->MotionMatchingState;
 
-		TArray<const UObject*> AssetsToSearch;
-		if (const UObject* DataBaseToSearch = SharedData->GetDatabase(Binding))
+		bool bNeedsReset = false;
+#if WITH_EDITOR
+		if (EAsyncBuildIndexResult::Success != FAsyncPoseSearchDatabasesManagement::RequestAsyncBuildIndex(MotionMatchingState.CurrentSearchResult.Database.Get(), ERequestAsyncBuildFlag::ContinueRequest))
 		{
-			AssetsToSearch.Add(DataBaseToSearch);
+			bNeedsReset = true;
 		}
-
-		if (AssetsToSearch.IsEmpty())
+		// in case this node is not updated, and MotionMatchingState.CurrentSearchResult.Database gets modified, we could end up with CurrentSearchResult being out of synch with the updated database, so we need to reset the state
+		else if (MotionMatchingState.CurrentSearchResult.IsValid() && MotionMatchingState.CurrentSearchResult.PoseIdx >= MotionMatchingState.CurrentSearchResult.Database->GetSearchIndex().GetNumPoses())
 		{
-			// @todo: LOG badness
-			return;
+			bNeedsReset = true;
 		}
+	#endif // WITH_EDITOR
 
-		IBlendStack::FGraphRequestPtr ActiveGraphRequest;
-		BlendStackBinding.GetActiveGraph(Context, ActiveGraphRequest);
-
-		FPoseSearchContinuingProperties ContinuingProperties;
-		if (ActiveGraphRequest)
+		// If we just became relevant and haven't been initialized yet, then reset motion matching state, otherwise update the asset time using the player node.
+		if (bNeedsReset)
 		{
-			check(ActiveGraphRequest->AnimationGraph);
-
-			for (const FInstancedStruct& Payload : ActiveGraphRequest->Payloads)
-			{
-				if (const FAnimNextPlayAnimPayload* PlayAnimPayload = Payload.GetPtr<FAnimNextPlayAnimPayload>())
-				{
-					ContinuingProperties.PlayingAsset = PlayAnimPayload->AnimationObject;
-
-					// it'd be nice to be able to have the value from the boud payload instead of having to ask to the ITimeline
-					// ContinuingProperties.PlayingAssetAccumulatedTime = PlayAnimPayload->StartPosition;
-					
-					TTraitBinding<ITimeline> TimelineTrait;
-					if (Binding.GetStackInterface(TimelineTrait))
-					{
-						const FTimelineProgress TimelineProgress = TimelineTrait.GetProgress(Context);
-						ContinuingProperties.PlayingAssetAccumulatedTime = TimelineProgress.GetPosition();
-					}
+			MotionMatchingState.Reset();
+			
+			// @todo: implement me!
+			//BlendStackBinding.ResetBlendStack();
+		}
+		else
+		{
+			// We adjust the motion matching state asset time to the current player node's asset time. This is done 
+			// because the player node may have ticked more or less time than we expected due to variable dt or the 
+			// dynamic playback rate adjustment and as such the motion matching state does not update by itself
+			const FTimelineProgress TimelineProgress = TimelineTrait.GetProgress(Context);
+			MotionMatchingState.AdjustAssetTime(TimelineProgress.GetPosition());
+		}
 
-					break;
-				}
-			}
+		TArray<TObjectPtr<const UPoseSearchDatabase>> DatabasesToSearch;
+		if (const UPoseSearchDatabase* DataBase = Cast<UPoseSearchDatabase>(SharedData->GetDatabase(Binding)))
+		{
+			DatabasesToSearch.Add(DataBase);
 		}
 
-		const UObject* AnimContext = nullptr;
-		if (const FAnimNextModuleInstance* ModuleInstance = Context.GetRootGraphInstance().GetModuleInstance())
+		if (DatabasesToSearch.IsEmpty())
 		{
-			AnimContext = ModuleInstance->Object;
+			UE_LOG(LogPoseSearch, Error, TEXT("FMotionMatchingTrait::PreUpdate, No database assets provided for motion matching."));
 		}
+		else
+		{
+			// code snippet on how to retrieve the currently playing animation from the blendstack via FAnimNextPlayAnimPayload
+			
+			//IBlendStack::FGraphRequestPtr ActiveGraphRequest;
+			//BlendStackBinding.GetActiveGraph(Context, ActiveGraphRequest);
+
+			//FPoseSearchContinuingProperties ContinuingProperties;
+			//if (ActiveGraphRequest)
+			//{
+			//	check(ActiveGraphRequest->AnimationGraph);
+
+			//	for (const FInstancedStruct& Payload : ActiveGraphRequest->Payloads)
+			//	{
+			//		if (const FAnimNextPlayAnimPayload* PlayAnimPayload = Payload.GetPtr<FAnimNextPlayAnimPayload>())
+			//		{
+			//			ContinuingProperties.PlayingAsset = PlayAnimPayload->AnimationObject;
+			//			break;
+			//		}
+			//	}
+			//}
+
+			const float DeltaTime = TraitState.GetDeltaTime();
+			const UObject* AnimContext = nullptr;
+			if (const FAnimNextModuleInstance* ModuleInstance = Context.GetRootGraphInstance().GetModuleInstance())
+			{
+				AnimContext = ModuleInstance->Object;
+			}
 
-		const TArrayView<const UObject*> AnimContexts = MakeArrayView(&AnimContext, 1);
-		const TArrayView<const FRole> Roles = MakeArrayView(&DefaultRole, 1);
-		const TArrayView<const UE::PoseSearch::IPoseHistory*> PoseHistories = MakeArrayView(&PoseHistory, 1);
-		const FPoseSearchFutureProperties FutureProperties;
+			UPoseSearchLibrary::UpdateMotionMatchingState(
+				AnimContext,
+				PoseHistory,
+				DatabasesToSearch,
+				DeltaTime,
+				SharedData->GetPoseJumpThresholdTime(Binding),
+				SharedData->GetPoseReselectHistory(Binding),
+				SharedData->GetbShouldSearch(Binding) ? SharedData->GetSearchThrottleTime(Binding) : UE_BIG_NUMBER,
+				SharedData->GetPlayRate(Binding),
+				MotionMatchingState,
+				SharedData->GetNextUpdateInterruptMode(Binding),
+				SharedData->GetbShouldUseCachedChannelData(Binding)
+				// @todo: add debug drawing suport
+//#if ENABLE_ANIM_DEBUG
+//				, CVarAnimNodeMotionMatchingDrawQuery.GetValueOnAnyThread()
+//				, CVarAnimNodeMotionMatchingDrawCurResult.GetValueOnAnyThread()
+//#endif // ENABLE_ANIM_DEBUG
+			);
+		}
 
-		const FSearchResult SearchResult = UPoseSearchLibrary::MotionMatch(AnimContexts, Roles, PoseHistories,  AssetsToSearch, ContinuingProperties, FutureProperties);
-		if (!SearchResult.bIsContinuingPoseSearch)
+		const bool bJumpToPose = MotionMatchingState.CurrentSearchResult.IsValid() && !MotionMatchingState.CurrentSearchResult.bIsContinuingPoseSearch;
+		if (bJumpToPose)
 		{
-			if (const UE::PoseSearch::FSearchIndexAsset* SearchIndexAsset = SearchResult.GetSearchIndexAsset())
+			if (const UE::PoseSearch::FSearchIndexAsset* SearchIndexAsset = MotionMatchingState.CurrentSearchResult.GetSearchIndexAsset())
 			{
-				if (const FPoseSearchDatabaseAnimationAssetBase* DatabaseAnimationAssetBase = SearchResult.Database->GetDatabaseAnimationAsset<FPoseSearchDatabaseAnimationAssetBase>(*SearchIndexAsset))
+				if (const FPoseSearchDatabaseAnimationAssetBase* DatabaseAnimationAssetBase = MotionMatchingState.CurrentSearchResult.Database->GetDatabaseAnimationAsset<FPoseSearchDatabaseAnimationAssetBase>(*SearchIndexAsset))
 				{
 					if (UAnimSequence* AnimSequence = Cast<UAnimSequence>(DatabaseAnimationAssetBase->GetAnimationAsset()))
 					{
@@ -122,7 +168,7 @@ namespace UE::AnimNext
 						FAnimNextPlayAnimPayload& NewPlayAnimPayload = NewPayload.GetMutable<FAnimNextPlayAnimPayload>();
 						NewPlayAnimPayload.AnimationObject = AnimSequence;
 						NewPlayAnimPayload.PlayRate = 1.f;
-						NewPlayAnimPayload.StartPosition = SearchResult.AssetTime;
+						NewPlayAnimPayload.StartPosition = MotionMatchingState.CurrentSearchResult.AssetTime;
 						NewPlayAnimPayload.IsLooping = SearchIndexAsset->IsLooping();
 
 						BlendStackBinding.PushGraph(Context, MoveTemp(NewGraphRequest));
diff --git a/Engine/Plugins/Experimental/AnimNextPoseSearch/Source/AnimNextPoseSearch/Private/MotionMatchingTrait.h b/Engine/Plugins/Experimental/AnimNextPoseSearch/Source/AnimNextPoseSearch/Private/MotionMatchingTrait.h
--- a/Engine/Plugins/Experimental/AnimNextPoseSearch/Source/AnimNextPoseSearch/Private/MotionMatchingTrait.h
+++ b/Engine/Plugins/Experimental/AnimNextPoseSearch/Source/AnimNextPoseSearch/Private/MotionMatchingTrait.h
@@ -23,6 +23,7 @@ struct FMotionMatchingTraitSharedData : public FAnimNextTraitSharedData
 	FName EntryPoint;
 
 	// @todo: support array of databases
+	// @todo: perhaps support databases AND choosers to be able to access chooser metadata =);>
 	// database to search from
 	UPROPERTY(EditAnywhere, Category = "Default")
 	TObjectPtr<const UObject> Database;
@@ -30,12 +31,50 @@ struct FMotionMatchingTraitSharedData : public FAnimNextTraitSharedData
 	UPROPERTY(EditAnywhere, Category = "Default")
 	FAlphaBlendArgs BlendArguments;
 
+	// Don't jump to poses of the same segment that are within the interval this many seconds away from the continuing pose.
+	UPROPERTY(EditAnywhere, Category = Settings, meta = (PinHiddenByDefault))
+	FFloatInterval PoseJumpThresholdTime = FFloatInterval(0.f, 0.f);
+
+	// Prevent re-selection of poses that have been selected previously within this much time (in seconds) in the past. This is across all animation segments that have been selected within this time range.
+	UPROPERTY(EditAnywhere, Category = Settings, meta = (PinHiddenByDefault, ClampMin = "0"))
+	float PoseReselectHistory = 0.3f;
+	
+	// If set to false, the motion matching node will perform a search only if the continuing pose is invalid. This is useful if you want to stagger searches of different nodes for performance reasons
+	UPROPERTY(EditAnywhere, Category = Settings, meta = (PinHiddenByDefault))
+	bool bShouldSearch = true;
+
+	// Minimum amount of time to wait between searching for a new pose segment. It allows users to define how often the system searches, default for locomotion is searching every update, but you may only want to search once for other situations, like jump.
+	UPROPERTY(EditAnywhere, Category = Settings, meta = (PinHiddenByDefault, ClampMin="0"))
+	float SearchThrottleTime = 0.f;
+
+	// Effective range of play rate that can be applied to the animations to account for discrepancies in estimated velocity between the movement model and the animation.
+	UPROPERTY(EditAnywhere, Category = Settings, meta = (PinHiddenByDefault, ClampMin = "0.2", ClampMax = "3.0", UIMin = "0.2", UIMax = "3.0"))
+	FFloatInterval PlayRate = FFloatInterval(1.f, 1.f);
+
+	// If set to true, the search of multiple databases with different schemas will try to share pose features data calculated during query build
+	// the idea is to be able to share as much as possible the continuing pose features vector across different schemas (and potentially improve performances)
+	// defaulted to false to preserve behavior backward compatibility
+	UPROPERTY(EditAnywhere, Category = Settings, meta = (PinHiddenByDefault))
+	bool bShouldUseCachedChannelData = false;
+
+	// @todo: should this be in FMotionMatchingState?
+	UPROPERTY(EditAnywhere, Category = Settings, meta = (PinHiddenByDefault))
+	// Applied EPoseSearchInterruptMode on the next update that controls the continuing pose search eveluation. This is set back to EPoseSearchInterruptMode::DoNotInterrupt after each update.
+	EPoseSearchInterruptMode NextUpdateInterruptMode = EPoseSearchInterruptMode::DoNotInterrupt;
+
 	// Latent pin support boilerplate
 	#define TRAIT_LATENT_PROPERTIES_ENUMERATOR(GeneratorMacro) \
 		GeneratorMacro(AnimationGraph) \
 		GeneratorMacro(EntryPoint) \
 		GeneratorMacro(Database) \
 		GeneratorMacro(BlendArguments) \
+		GeneratorMacro(PoseJumpThresholdTime) \
+		GeneratorMacro(PoseReselectHistory) \
+		GeneratorMacro(bShouldSearch) \
+		GeneratorMacro(SearchThrottleTime) \
+		GeneratorMacro(PlayRate) \
+		GeneratorMacro(bShouldUseCachedChannelData) \
+		GeneratorMacro(NextUpdateInterruptMode) \
 
 	GENERATE_TRAIT_LATENT_PROPERTIES(FMotionMatchingTraitSharedData, TRAIT_LATENT_PROPERTIES_ENUMERATOR)
 	#undef TRAIT_LATENT_PROPERTIES_ENUMERATOR
@@ -52,6 +91,7 @@ struct FMotionMatchingTrait : FAdditiveTrait, IUpdate
 		
 	struct FInstanceData : FTrait::FInstanceData
 	{
+		FMotionMatchingState MotionMatchingState;
 	};
 
 	// IUpdate impl
