AnimNextPoseSearch - PlayRate support

diff --git a/Engine/Plugins/Experimental/AnimNextPoseSearch/Source/AnimNextPoseSearch/Private/MotionMatchingTrait.cpp b/Engine/Plugins/Experimental/AnimNextPoseSearch/Source/AnimNextPoseSearch/Private/MotionMatchingTrait.cpp
--- a/Engine/Plugins/Experimental/AnimNextPoseSearch/Source/AnimNextPoseSearch/Private/MotionMatchingTrait.cpp
+++ b/Engine/Plugins/Experimental/AnimNextPoseSearch/Source/AnimNextPoseSearch/Private/MotionMatchingTrait.cpp
@@ -167,7 +167,7 @@ namespace UE::AnimNext
 						NewPayload.InitializeAs<FAnimNextPlayAnimPayload>();
 						FAnimNextPlayAnimPayload& NewPlayAnimPayload = NewPayload.GetMutable<FAnimNextPlayAnimPayload>();
 						NewPlayAnimPayload.AnimationObject = AnimSequence;
-						NewPlayAnimPayload.PlayRate = 1.f;
+						NewPlayAnimPayload.PlayRate = MotionMatchingState.WantedPlayRate;
 						NewPlayAnimPayload.StartPosition = MotionMatchingState.CurrentSearchResult.AssetTime;
 						NewPlayAnimPayload.IsLooping = SearchIndexAsset->IsLooping();
 
@@ -176,5 +176,9 @@ namespace UE::AnimNext
 				}
 			}
 		}
+
+		// @todo: implement me!
+		// FAnimNode_BlendStack_Standalone::UpdatePlayRate(MotionMatchingState.WantedPlayRate * PlayRateMultiplier);
+		// FAnimNode_BlendStack_Standalone::UpdateBlendspaceParameters(GetBlendspaceUpdateMode(), GetBlendspaceParameters());
 	}
 }
