Registered AlphaBlendArgs at RigVM registry to avoid a crash at pin connection check

#rnx

diff --git a/Engine/Plugins/Experimental/AnimNextPoseSearch/Source/AnimNextPoseSearch/Private/Module.cpp b/Engine/Plugins/Experimental/AnimNextPoseSearch/Source/AnimNextPoseSearch/Private/Module.cpp
--- a/Engine/Plugins/Experimental/AnimNextPoseSearch/Source/AnimNextPoseSearch/Private/Module.cpp
+++ b/Engine/Plugins/Experimental/AnimNextPoseSearch/Source/AnimNextPoseSearch/Private/Module.cpp
@@ -6,6 +6,7 @@
 #include "Modules/ModuleManager.h"
 #include "RigVMCore/RigVMRegistry.h"
 
+#include "AlphaBlend.h"
 #include "IAnimNextModuleInterface.h"
 #include "EvaluationVM/EvaluationVM.h"
 #include "PoseSearch/PoseSearchDatabase.h"
@@ -29,7 +30,8 @@ public:
 		static UScriptStruct* const AllowedStructTypes[] =
 		{
 			FPoseSearchQueryTrajectorySample::StaticStruct(),
-			FPoseSearchQueryTrajectory::StaticStruct()
+			FPoseSearchQueryTrajectory::StaticStruct(),
+			FAlphaBlendArgs::StaticStruct(),
 		};
 
 		FRigVMRegistry& RigVMRegistry = FRigVMRegistry::Get();
