AnimNextPoseSearch - FHistoryCollectorTrait cleanup

diff --git a/Engine/Plugins/Experimental/AnimNextPoseSearch/Source/AnimNextPoseSearch/Private/HistoryCollectorTrait.cpp b/Engine/Plugins/Experimental/AnimNextPoseSearch/Source/AnimNextPoseSearch/Private/HistoryCollectorTrait.cpp
--- a/Engine/Plugins/Experimental/AnimNextPoseSearch/Source/AnimNextPoseSearch/Private/HistoryCollectorTrait.cpp
+++ b/Engine/Plugins/Experimental/AnimNextPoseSearch/Source/AnimNextPoseSearch/Private/HistoryCollectorTrait.cpp
@@ -20,45 +20,6 @@ FAnimNextHistoryCollectorTask FAnimNextHistoryCollectorTask::Make(UE::PoseSearch
 
 namespace UE::AnimNext
 {
-	static void RemapPoseInternal(const FLODPoseStack& SourcePose, FCompactPose& TargetPose)
-	{
-		//SCOPE_CYCLE_COUNTER(STAT_AnimNext_RemapPose);
-
-		const FBoneContainer& BoneContainer = TargetPose.GetBoneContainer();
-		const FReferencePose& RefPose = SourcePose.GetRefPose();
-		const TArrayView<const FBoneIndexType> LODBoneIndexes = RefPose.GetLODBoneIndexToMeshBoneIndexMap(SourcePose.LODLevel);
-		const int32 NumLODBones = LODBoneIndexes.Num();
-
-		check(SourcePose.GetNumBones() == NumLODBones);
-
-		for (int32 LODBoneIndex = 0; LODBoneIndex < NumLODBones; ++LODBoneIndex)
-		{
-			// Reference pose holds a list of skeletal mesh bone indices sorted by LOD
-			const FMeshPoseBoneIndex MeshBoneIndex(LODBoneIndexes[LODBoneIndex]);
-
-			// Remap our skeletal mesh bone index into the skeleton bone index we output for
-			const FSkeletonPoseBoneIndex SkeletonBoneIndex = BoneContainer.GetSkeletonPoseIndexFromMeshPoseIndex(MeshBoneIndex);
-			if (SkeletonBoneIndex.IsValid())
-			{
-				// Remap our skeleton bone index into the compact pose bone index we output for
-				const FCompactPoseBoneIndex CompactBoneIndex = BoneContainer.GetCompactPoseIndexFromSkeletonPoseIndex(SkeletonBoneIndex);
-
-				if (ensure(CompactBoneIndex.IsValid()))	// We expect the skeleton bone to map to a valid compact pose bone
-				{
-					TargetPose[CompactBoneIndex] = SourcePose.LocalTransformsView[LODBoneIndex];
-				}
-				else
-				{
-					// This bone is part of the LOD but isn't part of the required bones
-				}
-			}
-			else
-			{
-				// This bone is part of the LOD but isn't part of the required bones
-			}
-		}
-	}
-
 	// @todo: revisit this code once a FCSPose<FLODPoseStack> is implemented!
 	struct FComponentSpacePoseProvider : public UE::PoseSearch::IComponentSpacePoseProvider
 	{
