AnimNextPoseSearch - removing  ControlRig dependency

diff --git a/Engine/Plugins/Experimental/AnimNextPoseSearch/Source/AnimNextPoseSearch/AnimNextPoseSearch.Build.cs b/Engine/Plugins/Experimental/AnimNextPoseSearch/Source/AnimNextPoseSearch/AnimNextPoseSearch.Build.cs
--- a/Engine/Plugins/Experimental/AnimNextPoseSearch/Source/AnimNextPoseSearch/AnimNextPoseSearch.Build.cs
+++ b/Engine/Plugins/Experimental/AnimNextPoseSearch/Source/AnimNextPoseSearch/AnimNextPoseSearch.Build.cs
@@ -11,7 +11,6 @@ namespace UnrealBuildTool.Rules
 				{
 					"AnimNext",
 					"AnimNextAnimGraph",
-					"ControlRig",
 					"Core",
 					"CoreUObject",
 					"Engine",
