AnimNext motion matching trait now continuously updates playrate instead of only when pushing a new anim

#tests win64 editor animsandbox lv_animnext_mover

diff --git a/Engine/Plugins/Animation/PoseSearch/Source/Runtime/Private/PoseSearchLibrary.cpp b/Engine/Plugins/Animation/PoseSearch/Source/Runtime/Private/PoseSearchLibrary.cpp
--- a/Engine/Plugins/Animation/PoseSearch/Source/Runtime/Private/PoseSearchLibrary.cpp
+++ b/Engine/Plugins/Animation/PoseSearch/Source/Runtime/Private/PoseSearchLibrary.cpp
@@ -765,6 +765,7 @@ void UPoseSearchLibrary::UpdateMotionMatchingState(
 	else
 	{
 		InOutMotionMatchingState.ElapsedPoseSearchTime += DeltaTime;
+		InOutMotionMatchingState.CurrentSearchResult.bIsContinuingPoseSearch = true;
 	}
 
 	// @todo: consider moving this into if (bSearch) to avoid calling SearchContext.GetCachedQuery if no search is required
diff --git a/Engine/Plugins/Experimental/AnimNextAnimGraph/Source/AnimNextAnimGraph/Public/TraitInterfaces/IBlendStack.h b/Engine/Plugins/Experimental/AnimNextAnimGraph/Source/AnimNextAnimGraph/Public/TraitInterfaces/IBlendStack.h
--- a/Engine/Plugins/Experimental/AnimNextAnimGraph/Source/AnimNextAnimGraph/Public/TraitInterfaces/IBlendStack.h
+++ b/Engine/Plugins/Experimental/AnimNextAnimGraph/Source/AnimNextAnimGraph/Public/TraitInterfaces/IBlendStack.h
@@ -107,6 +107,7 @@ namespace UE::AnimNext
 		// @see IBlendStack::GetActiveGraph
 		int32 GetActiveGraph(FExecutionContext& Context, IBlendStack::FGraphRequestPtr& OutGraphRequest) const
 		{
+			OutGraphRequest = nullptr; // Ensure request will be initialized
 			return GetInterface()->GetActiveGraph(Context, *this, OutGraphRequest);
 		}
 
diff --git a/Engine/Plugins/Experimental/AnimNextPoseSearch/Source/AnimNextPoseSearch/Private/MotionMatchingTrait.cpp b/Engine/Plugins/Experimental/AnimNextPoseSearch/Source/AnimNextPoseSearch/Private/MotionMatchingTrait.cpp
--- a/Engine/Plugins/Experimental/AnimNextPoseSearch/Source/AnimNextPoseSearch/Private/MotionMatchingTrait.cpp
+++ b/Engine/Plugins/Experimental/AnimNextPoseSearch/Source/AnimNextPoseSearch/Private/MotionMatchingTrait.cpp
@@ -26,6 +26,15 @@ namespace UE::AnimNext
 	GENERATE_ANIM_TRAIT_IMPLEMENTATION(FMotionMatchingTrait, TRAIT_INTERFACE_ENUMERATOR, NULL_ANIM_TRAIT_INTERFACE_ENUMERATOR, NULL_ANIM_TRAIT_EVENT_ENUMERATOR)
 	#undef TRAIT_INTERFACE_ENUMERATOR
 
+	void FMotionMatchingTrait::OnBecomeRelevant(FUpdateTraversalContext& Context, const TTraitBinding<IUpdate>& Binding, const FTraitUpdateState& TraitState) const
+	{
+		IUpdate::OnBecomeRelevant(Context, Binding, TraitState);
+
+		FInstanceData* InstanceData = Binding.GetInstanceData<FInstanceData>();
+		check(InstanceData);
+		InstanceData->GraphPlaybackRate = 1.0f;
+	}
+
 	void FMotionMatchingTrait::PreUpdate(FUpdateTraversalContext& Context, const TTraitBinding<IUpdate>& Binding, const FTraitUpdateState& TraitState) const
 	{
 		using namespace UE::PoseSearch;
@@ -84,6 +93,7 @@ namespace UE::AnimNext
 		if (bNeedsReset)
 		{
 			MotionMatchingState.Reset();
+			InstanceData->GraphPlaybackRate = 1.0f;
 			
 			// @todo: implement me!
 			//BlendStackBinding.ResetBlendStack();
@@ -104,26 +114,6 @@ namespace UE::AnimNext
 		}
 		else
 		{
-			// code snippet on how to retrieve the currently playing animation from the blendstack via FAnimNextPlayAnimPayload
-			
-			//IBlendStack::FGraphRequestPtr ActiveGraphRequest;
-			//BlendStackBinding.GetActiveGraph(Context, ActiveGraphRequest);
-
-			//FPoseSearchContinuingProperties ContinuingProperties;
-			//if (ActiveGraphRequest)
-			//{
-			//	check(ActiveGraphRequest->AnimationGraph);
-
-			//	for (const FInstancedStruct& Payload : ActiveGraphRequest->Payloads)
-			//	{
-			//		if (const FAnimNextPlayAnimPayload* PlayAnimPayload = Payload.GetPtr<FAnimNextPlayAnimPayload>())
-			//		{
-			//			ContinuingProperties.PlayingAsset = PlayAnimPayload->AnimationObject;
-			//			break;
-			//		}
-			//	}
-			//}
-
 			const float DeltaTime = TraitState.GetDeltaTime();
 			const UObject* AnimContext = nullptr;
 			if (const FAnimNextModuleInstance* ModuleInstance = Context.GetRootGraphInstance().GetModuleInstance())
@@ -182,13 +172,33 @@ namespace UE::AnimNext
 						NewGraphRequest.AnimationGraph = AnimationGraph;
 
 						BlendStackBinding.PushGraph(Context, MoveTemp(NewGraphRequest));
+						InstanceData->GraphPlaybackRate = MotionMatchingState.WantedPlayRate;
+					}
+				}
+			}
+		}
+		else if(!FMath::IsNearlyEqual(InstanceData->GraphPlaybackRate, MotionMatchingState.WantedPlayRate))
+		{
+			// Sync latest motion matching wanted playrate to current graph
+			IBlendStack::FGraphRequestPtr ActiveGraphRequest;
+			BlendStackBinding.GetActiveGraph(Context, ActiveGraphRequest);
+
+			if (ActiveGraphRequest)
+			{
+				check(ActiveGraphRequest->AnimationGraph);
+				for (FStructView& Payload : ActiveGraphRequest->GraphPayload.Get())
+				{
+					if (FAnimNextNativeDataInterface_AnimSequencePlayer* PlayAnimPayload = Payload.GetPtr<FAnimNextNativeDataInterface_AnimSequencePlayer>())
+					{
+						PlayAnimPayload->PlayRate = MotionMatchingState.WantedPlayRate;
+						InstanceData->GraphPlaybackRate = MotionMatchingState.WantedPlayRate;
+						break;
 					}
 				}
 			}
 		}
 
 		// @todo: implement me!
-		// FAnimNode_BlendStack_Standalone::UpdatePlayRate(MotionMatchingState.WantedPlayRate * PlayRateMultiplier);
 		// FAnimNode_BlendStack_Standalone::UpdateBlendspaceParameters(GetBlendspaceUpdateMode(), GetBlendspaceParameters());
 	}
 }
diff --git a/Engine/Plugins/Experimental/AnimNextPoseSearch/Source/AnimNextPoseSearch/Private/MotionMatchingTrait.h b/Engine/Plugins/Experimental/AnimNextPoseSearch/Source/AnimNextPoseSearch/Private/MotionMatchingTrait.h
--- a/Engine/Plugins/Experimental/AnimNextPoseSearch/Source/AnimNextPoseSearch/Private/MotionMatchingTrait.h
+++ b/Engine/Plugins/Experimental/AnimNextPoseSearch/Source/AnimNextPoseSearch/Private/MotionMatchingTrait.h
@@ -84,10 +84,14 @@ struct FMotionMatchingTrait : FAdditiveTrait, IUpdate
 	struct FInstanceData : FTrait::FInstanceData
 	{
 		FMotionMatchingState MotionMatchingState;
+		
+		// Cached version of the last playback rate set on the active graph
+		float GraphPlaybackRate = 1.0f;
 	};
 
 	// IUpdate impl
 	virtual void PreUpdate(FUpdateTraversalContext& Context, const TTraitBinding<IUpdate>& Binding, const FTraitUpdateState& TraitState) const override;
+	virtual void OnBecomeRelevant(FUpdateTraversalContext& Context, const TTraitBinding<IUpdate>& Binding, const FTraitUpdateState& TraitState) const override;
 };
 
 } // namespace UE::AnimNext
