Anim Next:
-Add support for blendspaces in Motion Matching.
-Make start time normalized for blendspaces
-Fix weight calculation for blendspaces
-Added factory support for blendspaces

#virtualized

diff --git a/Engine/Plugins/Experimental/AnimNextAnimGraph/Config/BaseAnimNextAnimGraph.ini b/Engine/Plugins/Experimental/AnimNextAnimGraph/Config/BaseAnimNextAnimGraph.ini
--- a/Engine/Plugins/Experimental/AnimNextAnimGraph/Config/BaseAnimNextAnimGraph.ini
+++ b/Engine/Plugins/Experimental/AnimNextAnimGraph/Config/BaseAnimNextAnimGraph.ini
@@ -2,6 +2,7 @@
 [/Script/AnimNextAnimGraph.AnimNextAnimGraphSettings]
 DefaultRunGraphHost=/AnimNextAnimGraph/FactoryGraphs/AG_DefaultGraphHost.AG_DefaultGraphHost
 +AssetGraphMappings=(AssetType="/Script/Engine.AnimSequence",AnimationGraph="/AnimNextAnimGraph/FactoryGraphs/AG_AnimSequencePlayer.AG_AnimSequencePlayer",Variable="AnimSequence")
++AssetGraphMappings=(AssetType="/Script/Engine.BlendSpace",AnimationGraph="/AnimNextAnimGraph/FactoryGraphs/AG_BlendSpacePlayer.AG_BlendSpacePlayer",Variable="BlendSpace")
 
 [CoreRedirects]
 ; Traits rename
diff --git a/Engine/Plugins/Experimental/AnimNextAnimGraph/Source/AnimNextAnimGraph/Private/GraphInterfaces/AnimNextNativeDataInterface_BlendSpacePlayer.cpp b/Engine/Plugins/Experimental/AnimNextAnimGraph/Source/AnimNextAnimGraph/Private/GraphInterfaces/AnimNextNativeDataInterface_BlendSpacePlayer.cpp
new file mode 100644
--- /dev/null
+++ b/Engine/Plugins/Experimental/AnimNextAnimGraph/Source/AnimNextAnimGraph/Private/GraphInterfaces/AnimNextNativeDataInterface_BlendSpacePlayer.cpp
@@ -0,0 +1,9 @@
+// Copyright Epic Games, Inc. All Rights Reserved.
+
+#include "GraphInterfaces/AnimNextNativeDataInterface_BlendSpacePlayer.h"
+#include "Animation/BlendSpace.h"
+
+void FAnimNextNativeDataInterface_BlendSpacePlayer::BindToFactoryObject(const FBindToFactoryObjectContext& InContext)
+{
+	BlendSpace = CastChecked<UBlendSpace>(InContext.FactoryObject);
+}
diff --git a/Engine/Plugins/Experimental/AnimNextAnimGraph/Source/AnimNextAnimGraph/Private/Traits/BlendSpacePlayer.cpp b/Engine/Plugins/Experimental/AnimNextAnimGraph/Source/AnimNextAnimGraph/Private/Traits/BlendSpacePlayer.cpp
--- a/Engine/Plugins/Experimental/AnimNextAnimGraph/Source/AnimNextAnimGraph/Private/Traits/BlendSpacePlayer.cpp
+++ b/Engine/Plugins/Experimental/AnimNextAnimGraph/Source/AnimNextAnimGraph/Private/Traits/BlendSpacePlayer.cpp
@@ -72,7 +72,7 @@ namespace UE::AnimNext
 				// we have no weight but still need to pop, so still make a task that will get ignored in the next step 
 				Context.AppendTask(FAnimNextBlendTwoKeyframesTask::Make(0.0));
 			}
-			Context.AppendTask(FAnimNextBlendTwoKeyframesTask::Make(BlendWeight0));
+			Context.AppendTask(FAnimNextBlendTwoKeyframesTask::Make(1.0 - BlendWeight0));
 		}
 		else if (InstanceData->NumChildren == 2)
 		{
@@ -101,7 +101,7 @@ namespace UE::AnimNext
 		{
 			InstanceData->BlendSpace->InitializeFilter(&InstanceData->BlendFilter);
 
-			FVector BlendParameters(SharedData->XAxisSamplePoint, SharedData->YAxisSamplePoint, 0.0f);
+			FVector BlendParameters(SharedData->GetXAxisSamplePoint(Binding), SharedData->GetYAxisSamplePoint(Binding), 0.0f);
 
 			// reset all the blend samples for the first frame
 			if (InstanceData->BlendSpace->GetSamplesFromBlendInput(BlendParameters, InstanceData->BlendSamplesData, InstanceData->CachedTriangulationIndex, true))
@@ -127,9 +127,13 @@ namespace UE::AnimNext
 						NewPayload.InitializeAs<FAnimNextNativeDataInterface_AnimSequencePlayer>();
 						FAnimNextNativeDataInterface_AnimSequencePlayer& NewNativeInterface = NewPayload.GetMutable<FAnimNextNativeDataInterface_AnimSequencePlayer>();
 						NewNativeInterface.AnimSequence = CurSequence;
-						NewNativeInterface.PlayRate = 1.0f;
-						NewNativeInterface.StartPosition = 0.0f;
-						NewNativeInterface.Loop = true;
+
+						// @TODO: Playrate should be normalized too.
+						NewNativeInterface.PlayRate = SharedData->GetPlayRate(Binding);
+
+						// Assume start position is normalized.
+						NewNativeInterface.StartPosition = SharedData->GetStartPosition(Binding) * CurSequence->GetPlayLength();
+						NewNativeInterface.Loop = SharedData->GetbLoop(Binding);
 
 						TargetGraphState.GraphPayload.AddNative(MoveTemp(NewPayload));
 
@@ -223,7 +227,7 @@ namespace UE::AnimNext
 		const float DeltaTime = TraitState.GetDeltaTime(); 
 
 		// input data 
-		FVector BlendParameters(SharedData->XAxisSamplePoint, SharedData->YAxisSamplePoint, 0.0f);
+		FVector BlendParameters(SharedData->GetXAxisSamplePoint(Binding), SharedData->GetYAxisSamplePoint(Binding), 0.0f);
 		for (FSampleGraphState& SampleState : InstanceData->SampleGraphs)
 		{
 			SampleState.bSampledLastFrame = SampleState.bSampledThisFrame;
diff --git a/Engine/Plugins/Experimental/AnimNextAnimGraph/Source/AnimNextAnimGraph/Private/Traits/BlendSpacePlayer.h b/Engine/Plugins/Experimental/AnimNextAnimGraph/Source/AnimNextAnimGraph/Private/Traits/BlendSpacePlayer.h
--- a/Engine/Plugins/Experimental/AnimNextAnimGraph/Source/AnimNextAnimGraph/Private/Traits/BlendSpacePlayer.h
+++ b/Engine/Plugins/Experimental/AnimNextAnimGraph/Source/AnimNextAnimGraph/Private/Traits/BlendSpacePlayer.h
@@ -33,11 +33,26 @@ struct FAnimNextBlendSpacePlayerTraitSharedData : public FAnimNextTraitSharedDat
 	UPROPERTY(EditAnywhere, Category = "Default")
 	float YAxisSamplePoint = 0.0f;
 
+	/** The play rate multiplier at which this blend space plays. */
+	UPROPERTY(EditAnywhere, Category = "Default")
+	float PlayRate = 1.0f;
+
+	/** The time at which we should start playing this blend space. This is normalized in the [0,1] range. */
+	UPROPERTY(EditAnywhere, Category = "Default")
+	float StartPosition = 0.0f;
+
+	/** Whether to loop the animation */
+	UPROPERTY(EditAnywhere, Category = "Default")
+	bool bLoop = false;
+
 	// Latent pin support boilerplate
 	#define TRAIT_LATENT_PROPERTIES_ENUMERATOR(GeneratorMacro) \
 		GeneratorMacro(BlendSpace) \
 		GeneratorMacro(XAxisSamplePoint) \
 		GeneratorMacro(YAxisSamplePoint) \
+		GeneratorMacro(PlayRate) \
+		GeneratorMacro(StartPosition) \
+		GeneratorMacro(bLoop) \
 
 
 	GENERATE_TRAIT_LATENT_PROPERTIES(FAnimNextBlendSpacePlayerTraitSharedData, TRAIT_LATENT_PROPERTIES_ENUMERATOR)
diff --git a/Engine/Plugins/Experimental/AnimNextAnimGraph/Source/AnimNextAnimGraph/Public/GraphInterfaces/AnimNextNativeDataInterface_BlendSpacePlayer.h b/Engine/Plugins/Experimental/AnimNextAnimGraph/Source/AnimNextAnimGraph/Public/GraphInterfaces/AnimNextNativeDataInterface_BlendSpacePlayer.h
new file mode 100644
--- /dev/null
+++ b/Engine/Plugins/Experimental/AnimNextAnimGraph/Source/AnimNextAnimGraph/Public/GraphInterfaces/AnimNextNativeDataInterface_BlendSpacePlayer.h
@@ -0,0 +1,45 @@
+// Copyright Epic Games, Inc. All Rights Reserved.
+
+#pragma once
+
+#include "CoreMinimal.h"
+#include "DataInterface/AnimNextNativeDataInterface.h"
+#include "AnimNextNativeDataInterface_BlendSpacePlayer.generated.h"
+
+class UBlendSpace;
+
+/**
+* Native interface to built-in BlendSpace player graph at: /AnimNextAnimGraph/FactoryGraphs/AG_BlendSpacePlayer
+ */
+USTRUCT(BlueprintType, DisplayName = "AnimNext Blend Space Player Data")
+struct FAnimNextNativeDataInterface_BlendSpacePlayer : public FAnimNextNativeDataInterface
+{
+	GENERATED_BODY()
+
+	// FAnimNextNativeDataInterface interface
+	ANIMNEXTANIMGRAPH_API virtual void BindToFactoryObject(const FBindToFactoryObjectContext& InContext) override;
+
+	// The animation object to play
+	UPROPERTY(BlueprintReadWrite, Category = "Blend Space")	// Not editable to hide it in cases where we would end up duplicating factory source asset
+	TObjectPtr<const UBlendSpace> BlendSpace;
+
+	// The location on the x-axis to sample.
+	UPROPERTY(EditAnywhere, BlueprintReadWrite, Category = "Blend Space", meta = (EnableCategories))
+	double XAxisSamplePoint = 0.0f;
+
+	// The location on the y-axis to sample.
+	UPROPERTY(EditAnywhere, BlueprintReadWrite, Category = "Blend Space", meta = (EnableCategories))
+	double YAxisSamplePoint = 0.0f;
+
+	// The play rate to use
+	UPROPERTY(EditAnywhere, BlueprintReadWrite, Category = "Blend Space", meta = (EnableCategories))
+	double PlayRate = 1.0f;
+
+	// The timeline's start position. This is normalized in the [0,1] range.
+	UPROPERTY(EditAnywhere, BlueprintReadWrite, Category = "Blend Space", meta = (EnableCategories))
+	double StartPosition = 0.0f;
+
+	// Whether to loop the animation
+	UPROPERTY(EditAnywhere, BlueprintReadWrite, Category = "Blend Space", meta = (EnableCategories))
+	bool bLoop = false;
+};
diff --git a/Engine/Plugins/Experimental/AnimNextPoseSearch/Source/AnimNextPoseSearch/Private/MotionMatchingTrait.cpp b/Engine/Plugins/Experimental/AnimNextPoseSearch/Source/AnimNextPoseSearch/Private/MotionMatchingTrait.cpp
--- a/Engine/Plugins/Experimental/AnimNextPoseSearch/Source/AnimNextPoseSearch/Private/MotionMatchingTrait.cpp
+++ b/Engine/Plugins/Experimental/AnimNextPoseSearch/Source/AnimNextPoseSearch/Private/MotionMatchingTrait.cpp
@@ -4,10 +4,12 @@
 
 #include "AnimNextAnimGraphSettings.h"
 #include "Animation/AnimSequence.h"
+#include "Animation/BlendSpace.h"
 #include "DataInterface/DataInterfaceStructAdapter.h"
 #include "IPoseHistory.h"
 #include "Graph/AnimNextGraphInstance.h"
 #include "GraphInterfaces/AnimNextNativeDataInterface_AnimSequencePlayer.h"
+#include "GraphInterfaces/AnimNextNativeDataInterface_BlendSpacePlayer.h"
 #include "Module/AnimNextModuleInstance.h"
 #include "PoseSearch/PoseSearchDatabase.h"
 #include "PoseSearch/PoseSearchDerivedData.h"
@@ -54,7 +56,6 @@ namespace UE::AnimNext
 
 		FInstanceData* InstanceData = Binding.GetInstanceData<FInstanceData>();
 		check(InstanceData);
-		InstanceData->GraphPlaybackRate = 1.0f;
 	}
 
 	void FMotionMatchingTrait::PreUpdate(FUpdateTraversalContext& Context, const TTraitBinding<IUpdate>& Binding, const FTraitUpdateState& TraitState) const
@@ -115,7 +116,6 @@ namespace UE::AnimNext
 		if (bNeedsReset)
 		{
 			MotionMatchingState.Reset();
-			InstanceData->GraphPlaybackRate = 1.0f;
 			
 			// @todo: implement me!
 			//BlendStackBinding.ResetBlendStack();
@@ -125,7 +125,28 @@ namespace UE::AnimNext
 			// We adjust the motion matching state asset time to the current player node's asset time. This is done 
 			// because the player node may have ticked more or less time than we expected due to variable dt or the 
 			// dynamic playback rate adjustment and as such the motion matching state does not update by itself
-			const float TimelinePosition = TimelineTrait.GetState(Context).GetPosition();
+			const FTimelineState TimelineState = TimelineTrait.GetState(Context);
+			float TimelinePosition = TimelineState.GetPosition();
+
+			// @TODO: This is a hack, since we don't have a way to get the normalized timeline of a blendspace yet (pending syncing...?)
+			if (const UE::PoseSearch::FSearchIndexAsset* SearchIndexAsset = MotionMatchingState.CurrentSearchResult.GetSearchIndexAsset())
+			{
+				if (const FPoseSearchDatabaseAnimationAssetBase* DatabaseAnimationAssetBase = MotionMatchingState.CurrentSearchResult.Database->GetDatabaseAnimationAsset<FPoseSearchDatabaseAnimationAssetBase>(*SearchIndexAsset))
+				{
+					const UObject* AnimationAsset = DatabaseAnimationAssetBase->GetAnimationAsset();
+					if (const UBlendSpace* BlendSpace = Cast<UBlendSpace>(AnimationAsset))
+					{
+						// Convert to normalized time.
+						const float TimelineDuration = TimelineState.GetDuration();
+						if (TimelineDuration > 0.0f)
+						{
+                            // NOTE: This doesn't work because AnimNext normalized time does not much database indexing normalized time.
+							TimelinePosition /= TimelineState.GetDuration();
+						}
+					}
+				}
+			}
+
 			MotionMatchingState.AdjustAssetTime(TimelinePosition);
 		}
 
@@ -191,12 +212,30 @@ namespace UE::AnimNext
 					{
 						// Create a default 'native data interface' payload via which we expect to communicate with the graph 
 						FInstancedStruct NewPayload;
-						NewPayload.InitializeAs<FAnimNextNativeDataInterface_AnimSequencePlayer>();
-						FAnimNextNativeDataInterface_AnimSequencePlayer& NewNativeInterface = NewPayload.GetMutable<FAnimNextNativeDataInterface_AnimSequencePlayer>();
-						NewNativeInterface.AnimSequence = Cast<UAnimSequence>(AnimationAsset);
-						NewNativeInterface.PlayRate = DesiredPlayRate;
-						NewNativeInterface.StartPosition = MotionMatchingState.CurrentSearchResult.AssetTime;
-						NewNativeInterface.Loop = SearchIndexAsset->IsLooping();
+						// @TODO: Support asset->playload mapping in a more data-driven way
+						if (AnimationAsset->IsA<UBlendSpace>())
+						{
+							NewPayload.InitializeAs<FAnimNextNativeDataInterface_BlendSpacePlayer>();
+							FAnimNextNativeDataInterface_BlendSpacePlayer& NewNativeInterface = NewPayload.GetMutable<FAnimNextNativeDataInterface_BlendSpacePlayer>();
+							NewNativeInterface.BlendSpace = Cast<UBlendSpace>(AnimationAsset);
+
+							const FVector BlendSpaceParameters = SearchIndexAsset->GetBlendParameters();
+							NewNativeInterface.XAxisSamplePoint = BlendSpaceParameters.X;
+							NewNativeInterface.YAxisSamplePoint = BlendSpaceParameters.Y;
+
+							NewNativeInterface.PlayRate = DesiredPlayRate;
+							NewNativeInterface.StartPosition = MotionMatchingState.CurrentSearchResult.AssetTime;
+							NewNativeInterface.bLoop = SearchIndexAsset->IsLooping();
+						}
+						else
+						{
+							NewPayload.InitializeAs<FAnimNextNativeDataInterface_AnimSequencePlayer>();
+							FAnimNextNativeDataInterface_AnimSequencePlayer& NewNativeInterface = NewPayload.GetMutable<FAnimNextNativeDataInterface_AnimSequencePlayer>();
+							NewNativeInterface.AnimSequence = Cast<UAnimSequence>(AnimationAsset);
+							NewNativeInterface.PlayRate = DesiredPlayRate;
+							NewNativeInterface.StartPosition = MotionMatchingState.CurrentSearchResult.AssetTime;
+							NewNativeInterface.Loop = SearchIndexAsset->IsLooping();
+						}
 
 						NewGraphRequest.GraphPayload.AddNative(MoveTemp(NewPayload));
 
@@ -205,33 +244,38 @@ namespace UE::AnimNext
 						NewGraphRequest.AnimationGraph = AnimationGraph;
 
 						BlendStackBinding.PushGraph(Context, MoveTemp(NewGraphRequest));
-						InstanceData->GraphPlaybackRate = MotionMatchingState.WantedPlayRate;
 					}
 				}
 			}
 		}
-		else if(!FMath::IsNearlyEqual(InstanceData->GraphPlaybackRate, DesiredPlayRate))
-		{
-			// Sync latest motion matching wanted playrate to current graph
-			IBlendStack::FGraphRequestPtr ActiveGraphRequest;
-			BlendStackBinding.GetActiveGraph(Context, ActiveGraphRequest);
 
-			if (ActiveGraphRequest)
+		// @TODO: This should be automatic and not manually piped per asset type.
+		// Sync latest motion matching state to current graph
+		IBlendStack::FGraphRequestPtr ActiveGraphRequest;
+		BlendStackBinding.GetActiveGraph(Context, ActiveGraphRequest);
+
+		if (ActiveGraphRequest)
+		{
+			check(ActiveGraphRequest->AnimationGraph);
+			for (FStructView& Payload : ActiveGraphRequest->GraphPayload.Get())
 			{
-				check(ActiveGraphRequest->AnimationGraph);
-				for (FStructView& Payload : ActiveGraphRequest->GraphPayload.Get())
+				if (FAnimNextNativeDataInterface_AnimSequencePlayer* PlayAnimPayload = Payload.GetPtr<FAnimNextNativeDataInterface_AnimSequencePlayer>())
+				{
+					PlayAnimPayload->PlayRate = DesiredPlayRate;
+					break;
+				}
+				else if (FAnimNextNativeDataInterface_BlendSpacePlayer* BlendSpacePayload = Payload.GetPtr<FAnimNextNativeDataInterface_BlendSpacePlayer>())
 				{
-					if (FAnimNextNativeDataInterface_AnimSequencePlayer* PlayAnimPayload = Payload.GetPtr<FAnimNextNativeDataInterface_AnimSequencePlayer>())
+					BlendSpacePayload->PlayRate = DesiredPlayRate;
+					if (SharedData->GetbUpdateBlendSpaceInputs(Binding))
 					{
-						PlayAnimPayload->PlayRate = DesiredPlayRate;
-						InstanceData->GraphPlaybackRate = DesiredPlayRate;
-						break;
+						BlendSpacePayload->XAxisSamplePoint = SharedData->GetXAxisSamplePoint(Binding);
+						BlendSpacePayload->YAxisSamplePoint = SharedData->GetYAxisSamplePoint(Binding);
 					}
+
+					break;
 				}
 			}
 		}
-
-		// @todo: implement me!
-		// FAnimNode_BlendStack_Standalone::UpdateBlendspaceParameters(GetBlendspaceUpdateMode(), GetBlendspaceParameters());
 	}
 }
diff --git a/Engine/Plugins/Experimental/AnimNextPoseSearch/Source/AnimNextPoseSearch/Private/MotionMatchingTrait.h b/Engine/Plugins/Experimental/AnimNextPoseSearch/Source/AnimNextPoseSearch/Private/MotionMatchingTrait.h
--- a/Engine/Plugins/Experimental/AnimNextPoseSearch/Source/AnimNextPoseSearch/Private/MotionMatchingTrait.h
+++ b/Engine/Plugins/Experimental/AnimNextPoseSearch/Source/AnimNextPoseSearch/Private/MotionMatchingTrait.h
@@ -28,6 +28,14 @@ struct FMotionMatchingTraitSharedData : public FAnimNextTraitSharedData
 	UPROPERTY(EditAnywhere, Category = Settings, meta = (PinHiddenByDefault))
 	FFloatInterval PoseJumpThresholdTime = FFloatInterval(0.f, 0.f);
 
+	/** The location on the x-axis to sample. Used if selection is a blend space. */
+	UPROPERTY(EditAnywhere, Category = "Blend Space")
+	float XAxisSamplePoint = 0.0f;
+
+	/** The location on the y-axis to sample. Used if selection is a blend space. */
+	UPROPERTY(EditAnywhere, Category = "Blend Space")
+	float YAxisSamplePoint = 0.0f;
+
 	// Prevent re-selection of poses that have been selected previously within this much time (in seconds) in the past. This is across all animation segments that have been selected within this time range.
 	UPROPERTY(EditAnywhere, Category = Settings, meta = (PinHiddenByDefault, ClampMin = "0"))
 	float PoseReselectHistory = 0.3f;
@@ -50,6 +58,10 @@ struct FMotionMatchingTraitSharedData : public FAnimNextTraitSharedData
 	UPROPERTY(EditAnywhere, Category = Settings, meta = (PinHiddenByDefault))
 	bool bShouldUseCachedChannelData = false;
 
+	// If set to true, blend space inputs are continuously set based on input pins.
+	UPROPERTY(EditAnywhere, Category = Settings)
+	bool bUpdateBlendSpaceInputs = false;
+
 	// @todo: should this be in FMotionMatchingState?
 	UPROPERTY(EditAnywhere, Category = Settings, meta = (PinHiddenByDefault))
 	// Applied EPoseSearchInterruptMode on the next update that controls the continuing pose search eveluation. This is set back to EPoseSearchInterruptMode::DoNotInterrupt after each update.
@@ -61,11 +73,14 @@ struct FMotionMatchingTraitSharedData : public FAnimNextTraitSharedData
 		GeneratorMacro(Databases) \
 		GeneratorMacro(BlendArguments) \
 		GeneratorMacro(PoseJumpThresholdTime) \
+		GeneratorMacro(XAxisSamplePoint) \
+		GeneratorMacro(YAxisSamplePoint) \
 		GeneratorMacro(PoseReselectHistory) \
 		GeneratorMacro(bShouldSearch) \
 		GeneratorMacro(SearchThrottleTime) \
 		GeneratorMacro(PlayRate) \
 		GeneratorMacro(bShouldUseCachedChannelData) \
+		GeneratorMacro(bUpdateBlendSpaceInputs) \
 		GeneratorMacro(NextUpdateInterruptMode) \
 
 	GENERATE_TRAIT_LATENT_PROPERTIES(FMotionMatchingTraitSharedData, TRAIT_LATENT_PROPERTIES_ENUMERATOR)
@@ -84,9 +99,6 @@ struct FMotionMatchingTrait : FAdditiveTrait, IUpdate
 	struct FInstanceData : FTrait::FInstanceData
 	{
 		FMotionMatchingState MotionMatchingState;
-		
-		// Cached version of the last playback rate set on the active graph
-		float GraphPlaybackRate = 1.0f;
 	};
 
 	// IUpdate impl
