Update motion matching trait to pass variable into another module

diff --git a/Engine/Plugins/Experimental/AnimNext/Source/AnimNext/Internal/Module/AnimNextModuleInstance.h b/Engine/Plugins/Experimental/AnimNext/Source/AnimNext/Internal/Module/AnimNextModuleInstance.h
--- a/Engine/Plugins/Experimental/AnimNext/Source/AnimNext/Internal/Module/AnimNextModuleInstance.h
+++ b/Engine/Plugins/Experimental/AnimNext/Source/AnimNext/Internal/Module/AnimNextModuleInstance.h
@@ -14,6 +14,7 @@
 #include "Module/AnimNextModuleInstanceComponent.h"
 #include "StructUtils/InstancedStruct.h"
 #include "Module/AnimNextModuleInitMethod.h"
+#include "Module/TaskRunLocation.h"
 
 #include "AnimNextModuleInstance.generated.h"
 
@@ -149,6 +150,20 @@ struct FAnimNextModuleInstance : public FAnimNextDataInterfaceInstance
 	ANIMNEXT_API TArrayView<UE::AnimNext::FModuleEventTickFunction> GetTickFunctions();
 
 	ANIMNEXT_API UE::AnimNext::FModuleHandle GetHandle() const { return Handle; }
+
+	// Queue a task to run at a particular module event
+	// @param	InModuleEventName	The name of the event in the module to run the supplied task relative to. If this is NAME_None, then the first valid event will be used.
+	// @param	InTaskFunction		The function to run
+	// @param	InLocation			Where to run the task, before or after
+	ANIMNEXT_API void QueueTask(FName InModuleEventName, TUniqueFunction<void(const UE::AnimNext::FModuleTaskContext&)>&& InTaskFunction, UE::AnimNext::ETaskRunLocation InLocation = UE::AnimNext::ETaskRunLocation::Before);
+
+	// Queue a task to run at a particular module event on some other module
+	// @param	InModuleHandle		The handle for the module to queue the task on
+	// @param	InModuleEventName	The name of the event in the module to run the supplied task relative to. If this is NAME_None, then the first valid event will be used.
+	// @param	InTaskFunction		The function to run
+	// @param	InLocation			Where to run the task, before or after
+	ANIMNEXT_API void QueueTaskOnOtherModule(const UE::AnimNext::FModuleHandle InOtherModuleHandle, FName InModuleEventName, TUniqueFunction<void(const UE::AnimNext::FModuleTaskContext&)>&& InTaskFunction, UE::AnimNext::ETaskRunLocation InLocation = UE::AnimNext::ETaskRunLocation::Before);
+
 private:
 	// Setup the entry
 	void Initialize();
diff --git a/Engine/Plugins/Experimental/AnimNext/Source/AnimNext/Private/AnimNextModuleImpl.cpp b/Engine/Plugins/Experimental/AnimNext/Source/AnimNext/Private/AnimNextModuleImpl.cpp
--- a/Engine/Plugins/Experimental/AnimNext/Source/AnimNext/Private/AnimNextModuleImpl.cpp
+++ b/Engine/Plugins/Experimental/AnimNext/Source/AnimNext/Private/AnimNextModuleImpl.cpp
@@ -12,6 +12,7 @@
 #include "Component/AnimNextComponent.h"
 #include "Curves/CurveFloat.h"
 #include "Module/AnimNextModuleInstance.h"
+#include "Module/ModuleHandle.h"
 #include "Modules/ModuleManager.h"
 #include "Param/AnimNextActorLocatorFragment.h"
 #include "Param/AnimNextComponentLocatorFragment.h"
@@ -65,6 +66,7 @@ namespace UE::AnimNext
 			FRigVMGraphFunctionHeader::StaticStruct(),
 			TBaseStructure<FGuid>::Get(),
 			FRigVMVariant::StaticStruct(),
+			FAnimNextModuleHandle::StaticStruct(),
 		};
 
 		RigVMRegistry.RegisterStructTypes(AllowedStructTypes);
diff --git a/Engine/Plugins/Experimental/AnimNext/Source/AnimNext/Private/Component/AnimNextComponent.cpp b/Engine/Plugins/Experimental/AnimNext/Source/AnimNext/Private/Component/AnimNextComponent.cpp
--- a/Engine/Plugins/Experimental/AnimNext/Source/AnimNext/Private/Component/AnimNextComponent.cpp
+++ b/Engine/Plugins/Experimental/AnimNext/Source/AnimNext/Private/Component/AnimNextComponent.cpp
@@ -949,3 +949,8 @@ void UAnimNextComponent::RemoveSubsequent(UObject* InObject, FTickFunction& InTi
 		Subsystem->RemoveDependency(this, InObject, InTickFunction, InEventName, UAnimNextWorldSubsystem::EDependency::Subsequent);
 	}
 }
+
+FAnimNextModuleHandle UAnimNextComponent::BlueprintGetModuleHandle() const
+{
+	return FAnimNextModuleHandle(ModuleHandle);
+}
\ No newline at end of file
diff --git a/Engine/Plugins/Experimental/AnimNext/Source/AnimNext/Private/DataInterface/AnimNextDataInterfaceInstance.cpp b/Engine/Plugins/Experimental/AnimNext/Source/AnimNext/Private/DataInterface/AnimNextDataInterfaceInstance.cpp
--- a/Engine/Plugins/Experimental/AnimNext/Source/AnimNext/Private/DataInterface/AnimNextDataInterfaceInstance.cpp
+++ b/Engine/Plugins/Experimental/AnimNext/Source/AnimNext/Private/DataInterface/AnimNextDataInterfaceInstance.cpp
@@ -1,4 +1,4 @@
-﻿// Copyright Epic Games, Inc. All Rights Reserved.
+// Copyright Epic Games, Inc. All Rights Reserved.
 
 #include "DataInterface/AnimNextDataInterfaceInstance.h"
 #include "Logging/StructuredLog.h"
@@ -596,6 +596,43 @@ EPropertyBagResult FAnimNextDataInterfaceInstance::GetVariableInternal(FName InV
 	return EPropertyBagResult::Success;
 }
 
+EPropertyBagResult FAnimNextDataInterfaceInstance::SetVariableInternal(FName InVariableName, const FAnimNextParamType& InType, TConstArrayView<uint8> InNewValue)
+{
+	// @TODO: Combine w/ GetVariable since they're almost identical.
+	using namespace UE::AnimNext;
+
+	const UPropertyBag* PropertyBag = Variables.GetPropertyBagStruct();
+	if(PropertyBag == nullptr)
+	{
+		return EPropertyBagResult::PropertyNotFound;
+	}
+
+	const FPropertyBagPropertyDesc* Desc = PropertyBag->FindPropertyDescByName(InVariableName);
+	if(Desc == nullptr)
+	{
+		if(HostInstance)
+		{
+			return HostInstance->SetVariableInternal(InVariableName, InType, InNewValue);
+		}
+		else
+		{
+			return EPropertyBagResult::PropertyNotFound;
+		}
+	}
+
+	const int32 DescIndex = Desc - &PropertyBag->GetPropertyDescs()[0];
+	uint8* Memory = ExtendedExecuteContext.ExternalVariableRuntimeData[DescIndex].Memory;
+
+	const FAnimNextParamType InternalType(Desc->ValueType, Desc->ContainerTypes.GetFirstContainerType(), Desc->ValueTypeObject);
+	if(InternalType != InType)
+	{
+		return Private::GetVariableFromMismatchedType(Desc, InternalType, InType, InNewValue.GetData(), Memory);
+	}
+
+	Desc->CachedProperty->CopyCompleteValue(Memory, InNewValue.GetData());
+	return EPropertyBagResult::Success;
+}
+
 FAnimNextDataInterfaceInstance::FAnimNextDataInterfaceInstance()
 {
 #if ANIMNEXT_TRACE_ENABLED
diff --git a/Engine/Plugins/Experimental/AnimNext/Source/AnimNext/Private/Module/AnimNextModuleInstance.cpp b/Engine/Plugins/Experimental/AnimNext/Source/AnimNext/Private/Module/AnimNextModuleInstance.cpp
--- a/Engine/Plugins/Experimental/AnimNext/Source/AnimNext/Private/Module/AnimNextModuleInstance.cpp
+++ b/Engine/Plugins/Experimental/AnimNext/Source/AnimNext/Private/Module/AnimNextModuleInstance.cpp
@@ -758,4 +758,65 @@ TArrayView<UE::AnimNext::FModuleEventTickFunction> FAnimNextModuleInstance::GetT
 	UE::AnimNext::FModuleWriteGuard PrerequisiteGuard(this);
 
 	return TickFunctions;
+}
+
+void FAnimNextModuleInstance::QueueTask(FName InModuleEventName, TUniqueFunction<void(const UE::AnimNext::FModuleTaskContext&)>&& InTaskFunction, UE::AnimNext::ETaskRunLocation InLocation)
+{
+	using namespace UE::AnimNext;
+
+	TSpscQueue<TUniqueFunction<void(const FModuleTaskContext&)>>* Queue = nullptr;
+	FModuleEventTickFunction* FoundTickFunction = nullptr;
+	if(TickFunctions.Num() > 0)
+	{
+		if(InModuleEventName == NAME_None)
+		{
+			// passing NAME_None means 'any tick function' in this context, so just use the first one
+			FoundTickFunction = &TickFunctions[0];
+		}
+		else
+		{
+			FoundTickFunction = TickFunctions.FindByPredicate([InModuleEventName](const FModuleEventTickFunction& InTickFunction)
+			{
+				return InTickFunction.EventName == InModuleEventName;
+			});
+		}
+	}
+
+	if (FoundTickFunction)
+	{
+		switch (InLocation)
+		{
+		case ETaskRunLocation::Before:
+			Queue = &FoundTickFunction->PreExecuteTasks;
+			break;
+		case ETaskRunLocation::After:
+			Queue = &FoundTickFunction->PostExecuteTasks;
+			break;
+		}
+	}
+
+	if (Queue)
+	{
+		Queue->Enqueue(MoveTemp(InTaskFunction));
+	}
+	else
+	{
+		UE_LOGFMT(LogAnimation, Warning, "QueueTask: Could not find event '{EventName}' in module '{ModuleName}'", InModuleEventName, GetDataInterfaceName());
+	}
+}
+
+void FAnimNextModuleInstance::QueueTaskOnOtherModule(const UE::AnimNext::FModuleHandle InOtherModuleHandle, FName InModuleEventName, TUniqueFunction<void(const UE::AnimNext::FModuleTaskContext&)>&& InTaskFunction, UE::AnimNext::ETaskRunLocation InLocation)
+{
+	if (Pool == nullptr)
+	{
+		return;
+	}
+
+	FAnimNextModuleInstance* OtherModuleInstance = Pool->TryGet(InOtherModuleHandle);
+	if (OtherModuleInstance == nullptr)
+	{
+		return;
+	}
+
+	OtherModuleInstance->QueueTask(InModuleEventName, MoveTemp(InTaskFunction), InLocation);
 }
\ No newline at end of file
diff --git a/Engine/Plugins/Experimental/AnimNext/Source/AnimNext/Private/Module/ModuleTaskContext.cpp b/Engine/Plugins/Experimental/AnimNext/Source/AnimNext/Private/Module/ModuleTaskContext.cpp
--- a/Engine/Plugins/Experimental/AnimNext/Source/AnimNext/Private/Module/ModuleTaskContext.cpp
+++ b/Engine/Plugins/Experimental/AnimNext/Source/AnimNext/Private/Module/ModuleTaskContext.cpp
@@ -35,4 +35,9 @@ void FModuleTaskContext::TryAccessComponent(UScriptStruct* InComponentType, TFun
 	InFunction(*Component);
 }
 
+FAnimNextModuleInstance* const FModuleTaskContext::GetModuleInstance() const
+{
+	return ModuleInstance;
+}
+
 }
diff --git a/Engine/Plugins/Experimental/AnimNext/Source/AnimNext/Public/Component/AnimNextComponent.h b/Engine/Plugins/Experimental/AnimNext/Source/AnimNext/Public/Component/AnimNextComponent.h
--- a/Engine/Plugins/Experimental/AnimNext/Source/AnimNext/Public/Component/AnimNextComponent.h
+++ b/Engine/Plugins/Experimental/AnimNext/Source/AnimNext/Public/Component/AnimNextComponent.h
@@ -233,6 +233,9 @@ public:
 	// @param	InEventName			The event to remove the dependency from
 	ANIMNEXT_API void RemoveSubsequent(UObject* InObject, FTickFunction& InTickFunction, FName InEventName);
 
+	UFUNCTION(BlueprintCallable, Category = "AnimNext", meta=(DisplayName = "Get Module Handle"))
+	ANIMNEXT_API FAnimNextModuleHandle BlueprintGetModuleHandle() const;
+
 private:
 	DECLARE_FUNCTION(execSetVariable);
 
diff --git a/Engine/Plugins/Experimental/AnimNext/Source/AnimNext/Public/DataInterface/AnimNextDataInterfaceInstance.h b/Engine/Plugins/Experimental/AnimNext/Source/AnimNext/Public/DataInterface/AnimNextDataInterfaceInstance.h
--- a/Engine/Plugins/Experimental/AnimNext/Source/AnimNext/Public/DataInterface/AnimNextDataInterfaceInstance.h
+++ b/Engine/Plugins/Experimental/AnimNext/Source/AnimNext/Public/DataInterface/AnimNextDataInterfaceInstance.h
@@ -1,4 +1,4 @@
-﻿// Copyright Epic Games, Inc. All Rights Reserved.
+// Copyright Epic Games, Inc. All Rights Reserved.
 
 #pragma once
 
@@ -58,15 +58,25 @@ struct ANIMNEXT_API FAnimNextDataInterfaceInstance
 	uint8* GetMemoryForVariable(int32 InVariableIndex, FName InVariableName, const FProperty* InVariableProperty) const;
 
 	// Get a variable's value given its name.
-	// @param	InVariableName		The name of the variable to get he value of
+	// @param	InVariableName		The name of the variable to get the value of
 	// @param	OutResult			Result that will be filled if no errors occur
-	// @return nullptr if the variable is not present
+	// @return see EPropertyBagResult
 	template<typename ValueType>
 	EPropertyBagResult GetVariable(FName InVariableName, ValueType& OutResult) const
 	{
 		return GetVariableInternal(InVariableName, FAnimNextParamType::GetType<ValueType>(), TArrayView<uint8>(reinterpret_cast<uint8*>(&OutResult), sizeof(ValueType)));
 	}
 
+	// Set a variable's value given its name.
+	// @param	InVariableName		The name of the variable to set the value of
+	// @param	OutResult			Result that will be filled if no errors occur
+	// @return see EPropertyBagResult
+	template<typename ValueType>
+	EPropertyBagResult SetVariable(const FName InVariableName, const ValueType& InNewValue)
+	{
+		return SetVariableInternal(InVariableName, FAnimNextParamType::GetType<ValueType>(), TConstArrayView<uint8>(reinterpret_cast<const uint8*>(&InNewValue), sizeof(ValueType)));
+	}
+
 	// Get the instance (graph, module etc.) that owns/hosts us
 	FAnimNextDataInterfaceInstance* GetHost() const { return HostInstance; }
 	
@@ -80,6 +90,8 @@ struct ANIMNEXT_API FAnimNextDataInterfaceInstance
 private:
 	// Helper function for GetVariable
 	EPropertyBagResult GetVariableInternal(FName InVariableName, const FAnimNextParamType& InType, TArrayView<uint8> OutResult) const;
+	// Helper function for SetVariable
+	EPropertyBagResult SetVariableInternal(const FName InVariableName, const FAnimNextParamType& InType, TConstArrayView<uint8> InNewValue);
 
 protected:
 	// Hard reference to the asset used to create this instance to ensure we can release it safely
diff --git a/Engine/Plugins/Experimental/AnimNext/Source/AnimNext/Public/Module/ModuleHandle.h b/Engine/Plugins/Experimental/AnimNext/Source/AnimNext/Public/Module/ModuleHandle.h
--- a/Engine/Plugins/Experimental/AnimNext/Source/AnimNext/Public/Module/ModuleHandle.h
+++ b/Engine/Plugins/Experimental/AnimNext/Source/AnimNext/Public/Module/ModuleHandle.h
@@ -4,6 +4,8 @@
 
 #include "AnimNextPoolHandle.h"
 
+#include "ModuleHandle.generated.h"
+
 struct FAnimNextModuleInstance;
 
 namespace UE::AnimNext
@@ -12,4 +14,20 @@ namespace UE::AnimNext
 // Opaque handle representing a module instance
 using FModuleHandle = TPoolHandle<FAnimNextModuleInstance>;
 
-}
\ No newline at end of file
+}
+
+USTRUCT(BlueprintType)
+struct FAnimNextModuleHandle
+{
+	GENERATED_BODY()
+
+	FAnimNextModuleHandle() = default;
+
+	explicit FAnimNextModuleHandle(const UE::AnimNext::FModuleHandle InModuleHandle)
+		: ModuleHandle(InModuleHandle)
+	{
+	}
+
+public:
+	UE::AnimNext::FModuleHandle ModuleHandle;
+};
\ No newline at end of file
diff --git a/Engine/Plugins/Experimental/AnimNext/Source/AnimNext/Public/Module/ModuleTaskContext.h b/Engine/Plugins/Experimental/AnimNext/Source/AnimNext/Public/Module/ModuleTaskContext.h
--- a/Engine/Plugins/Experimental/AnimNext/Source/AnimNext/Public/Module/ModuleTaskContext.h
+++ b/Engine/Plugins/Experimental/AnimNext/Source/AnimNext/Public/Module/ModuleTaskContext.h
@@ -41,6 +41,8 @@ public:
 		});
 	}
 
+	ANIMNEXT_API FAnimNextModuleInstance* const GetModuleInstance() const;
+
 private:
 	FModuleTaskContext(FAnimNextModuleInstance& InModuleInstance);
 
diff --git a/Engine/Plugins/Experimental/AnimNextPoseSearch/Source/AnimNextPoseSearch/Private/Module.cpp b/Engine/Plugins/Experimental/AnimNextPoseSearch/Source/AnimNextPoseSearch/Private/Module.cpp
--- a/Engine/Plugins/Experimental/AnimNextPoseSearch/Source/AnimNextPoseSearch/Private/Module.cpp
+++ b/Engine/Plugins/Experimental/AnimNextPoseSearch/Source/AnimNextPoseSearch/Private/Module.cpp
@@ -32,6 +32,7 @@ public:
 		{
 			FTransformTrajectorySample::StaticStruct(),
 			FTransformTrajectory::StaticStruct(),
+			FPoseSearchBlueprintResult::StaticStruct(),
 			FAlphaBlendArgs::StaticStruct(),
 			FPoseSearchInteractionAvailability::StaticStruct(),
 		};
diff --git a/Engine/Plugins/Experimental/AnimNextPoseSearch/Source/AnimNextPoseSearch/Private/MotionMatchingTrait.cpp b/Engine/Plugins/Experimental/AnimNextPoseSearch/Source/AnimNextPoseSearch/Private/MotionMatchingTrait.cpp
--- a/Engine/Plugins/Experimental/AnimNextPoseSearch/Source/AnimNextPoseSearch/Private/MotionMatchingTrait.cpp
+++ b/Engine/Plugins/Experimental/AnimNextPoseSearch/Source/AnimNextPoseSearch/Private/MotionMatchingTrait.cpp
@@ -13,6 +13,7 @@
 #include "GraphInterfaces/AnimNextNativeDataInterface_BlendSpacePlayer.h"
 #include "IPoseHistory.h"
 #include "Module/AnimNextModuleInstance.h"
+#include "Module/RigUnit_AnimNextModuleEvents.h"
 #include "PoseSearch/MultiAnimAsset.h"
 #include "PoseSearch/PoseSearchDatabase.h"
 #include "PoseSearch/PoseSearchDerivedData.h"
@@ -22,6 +23,7 @@
 #include "TraitInterfaces/IGraphFactory.h"
 #include "TraitInterfaces/ITimeline.h"
 #include "VisualLogger/VisualLogger.h"
+#include "TraitCore/NodeInstance.h"
 
 #if ENABLE_ANIM_DEBUG
 namespace UE::Private
@@ -339,6 +341,48 @@ namespace UE::AnimNext
 				}
 			}
 		}
+
+		FName ResultOutputVariableName = SharedData->GetMotionMatchingResultVariableName(Binding);
+		const FAnimNextModuleHandle ResultModuleHandle = SharedData->GetMotionMatchingResultModuleHandle(Binding);
+		if (ResultOutputVariableName != NAME_None && ResultModuleHandle.ModuleHandle.IsValid())
+		{
+			FAnimNextModuleInstance* ModuleInstance = Binding.GetTraitPtr().GetNodeInstance()->GetOwner().GetModuleInstance();
+
+			const UPoseSearchDatabase* Database = MotionMatchingState.CurrentSearchResult.Database.Get();
+
+			if (const FSearchIndexAsset* SearchIndexAsset = MotionMatchingState.CurrentSearchResult.GetSearchIndexAsset())
+			{
+				const UPoseSearchDatabase* CurrentResultDatabase = MotionMatchingState.CurrentSearchResult.Database.Get();
+				if (CurrentResultDatabase && CurrentResultDatabase->Schema)
+				{
+					if (const FPoseSearchDatabaseAnimationAssetBase* DatabaseAsset = CurrentResultDatabase->GetDatabaseAnimationAsset<FPoseSearchDatabaseAnimationAssetBase>(*SearchIndexAsset))
+					{
+						FPoseSearchBlueprintResult Result;
+						Result.SelectedAnimation = DatabaseAsset->GetAnimationAsset();
+						Result.SelectedTime = MotionMatchingState.CurrentSearchResult.AssetTime;
+						Result.bIsContinuingPoseSearch = MotionMatchingState.CurrentSearchResult.bIsContinuingPoseSearch;
+						Result.bLoop = SearchIndexAsset->IsLooping();
+						Result.bIsMirrored = SearchIndexAsset->IsMirrored();
+						Result.BlendParameters = SearchIndexAsset->GetBlendParameters();
+						Result.SelectedDatabase = CurrentResultDatabase;
+						Result.WantedPlayRate = MotionMatchingState.WantedPlayRate;
+
+						ModuleInstance->QueueTaskOnOtherModule(
+							ResultModuleHandle.ModuleHandle,
+							FRigUnit_AnimNextPrePhysicsEvent::EventName,
+							[Result, ResultOutputVariableName](const UE::AnimNext::FModuleTaskContext& TaskContext)
+								{
+									FAnimNextModuleInstance* OtherModuleInstance = TaskContext.GetModuleInstance();
+									if (OtherModuleInstance)
+									{
+										OtherModuleInstance->SetVariable(ResultOutputVariableName, Result);
+									}
+								},
+							UE::AnimNext::ETaskRunLocation::Before);
+					}
+				}
+			}
+		}
 	}
 
 	void FMotionMatchingTrait::PostEvaluate(FEvaluateTraversalContext& Context, const TTraitBinding<IEvaluate>& Binding) const
diff --git a/Engine/Plugins/Experimental/AnimNextPoseSearch/Source/AnimNextPoseSearch/Private/MotionMatchingTrait.h b/Engine/Plugins/Experimental/AnimNextPoseSearch/Source/AnimNextPoseSearch/Private/MotionMatchingTrait.h
--- a/Engine/Plugins/Experimental/AnimNextPoseSearch/Source/AnimNextPoseSearch/Private/MotionMatchingTrait.h
+++ b/Engine/Plugins/Experimental/AnimNextPoseSearch/Source/AnimNextPoseSearch/Private/MotionMatchingTrait.h
@@ -9,6 +9,7 @@
 #include "TraitCore/TraitSharedData.h"
 #include "TraitInterfaces/IEvaluate.h"
 #include "TraitInterfaces/IUpdate.h"
+#include "Module/ModuleHandle.h"
 #include "MotionMatchingTrait.generated.h"
 
 USTRUCT(meta = (DisplayName = "Motion Matching"))
@@ -65,6 +66,14 @@ struct FMotionMatchingTraitSharedData : public FAnimNextTraitSharedData
 	// Applied EPoseSearchInterruptMode on the next update that controls the continuing pose search eveluation. This is set back to EPoseSearchInterruptMode::DoNotInterrupt after each update.
 	EPoseSearchInterruptMode NextUpdateInterruptMode = EPoseSearchInterruptMode::DoNotInterrupt;
 
+	// If set, the motion matching result will be written to this variable
+	UPROPERTY(EditAnywhere, Category = Settings)
+	FName MotionMatchingResultVariableName = NAME_None;
+
+	// If set along with result variable name, the motion matching result will be written to the variable in this module
+	UPROPERTY(EditAnywhere, Category = Settings)
+	FAnimNextModuleHandle MotionMatchingResultModuleHandle;
+
 	UPROPERTY(EditAnywhere, Category = Interaction, meta = (PinHiddenByDefault))
 	TArray<FPoseSearchInteractionAvailability> Availabilities;
 
@@ -94,6 +103,8 @@ struct FMotionMatchingTraitSharedData : public FAnimNextTraitSharedData
 		GeneratorMacro(bShouldUseCachedChannelData) \
 		GeneratorMacro(bUpdateBlendSpaceInputs) \
 		GeneratorMacro(NextUpdateInterruptMode) \
+		GeneratorMacro(MotionMatchingResultVariableName) \
+		GeneratorMacro(MotionMatchingResultModuleHandle) \
 		GeneratorMacro(Availabilities) \
 		GeneratorMacro(bValidateResultAgainstAvailabilities) \
 		GeneratorMacro(bWarpUsingRootBone) \
