PoseSearch
- exposing MM interaction state as variable in MM trait
- made FPoseSearchBlueprintResult GC
- merged FPoseSearchInteractionBlueprintResult with FPoseSearchBlueprintResult
- fixed threading issue in UMotionMatchingAnimNodeLibrary::GetMotionMatchingSearchResult while databases are indexing

#rnx

diff --git a/Engine/Plugins/Animation/PoseSearch/Config/BasePoseSearch.ini b/Engine/Plugins/Animation/PoseSearch/Config/BasePoseSearch.ini
--- a/Engine/Plugins/Animation/PoseSearch/Config/BasePoseSearch.ini
+++ b/Engine/Plugins/Animation/PoseSearch/Config/BasePoseSearch.ini
@@ -28,6 +28,7 @@
 +FunctionRedirects=(OldName="PoseSearchHistoryCollectorAnimNodeLibrary.GetPoseHistoryNodeTrajectory", NewName="PoseSearchHistoryCollectorAnimNodeLibrary.GetPoseHistoryNodeTransformTrajectory")
 +FunctionRedirects=(OldName="PoseSearchHistoryCollectorAnimNodeLibrary.SetPoseHistoryNodeTrajectory", NewName="PoseSearchHistoryCollectorAnimNodeLibrary.SetPoseHistoryNodeTransformTrajectory")
 
-; Fix up pose history node stuff
+; Fix up pose history node and PoseSearchBlueprintResult stuff
 +PropertyRedirects=(OldName="AnimNode_PoseSearchHistoryCollector_Base.Trajectory", NewName="AnimNode_PoseSearchHistoryCollector_Base.TransformTrajectory")
 +PropertyRedirects=(OldName="AnimNode_PoseSearchHistoryCollector.Trajectory", NewName="AnimNode_PoseSearchHistoryCollector.TransformTrajectory")
++PropertyRedirects=(OldName="PoseSearchBlueprintResult.SelectedAnimation", NewName="PoseSearchBlueprintResult.SelectedAnim")
\ No newline at end of file
diff --git a/Engine/Plugins/Animation/PoseSearch/Source/Runtime/Private/AnimNode_MotionMatchingInteraction.cpp b/Engine/Plugins/Animation/PoseSearch/Source/Runtime/Private/AnimNode_MotionMatchingInteraction.cpp
--- a/Engine/Plugins/Animation/PoseSearch/Source/Runtime/Private/AnimNode_MotionMatchingInteraction.cpp
+++ b/Engine/Plugins/Animation/PoseSearch/Source/Runtime/Private/AnimNode_MotionMatchingInteraction.cpp
@@ -23,7 +23,7 @@
 void FAnimNode_MotionMatchingInteraction::Reset()
 {
 	Super::Reset();
-	CurrentResult = FPoseSearchInteractionBlueprintResult();
+	CurrentResult = FPoseSearchBlueprintResult();
 	MeshWithOffset = FTransform::Identity;
 	MeshWithoutOffset = FTransform::Identity;
 	CachedDeltaTime = 0.f;
@@ -65,18 +65,18 @@ void FAnimNode_MotionMatchingInteraction::UpdateAssetPlayer(const FAnimationUpda
 	{
 		// we're NOT providing the UMultiAnimAsset here, relying on the interaction subsystem to resolve the ContinuingProperties properly
 		FPoseSearchContinuingProperties ContinuingProperties;
-		ContinuingProperties.PlayingAsset = CurrentResult.SelectedAnimation; // GetAnimAsset();
+		ContinuingProperties.PlayingAsset = CurrentResult.SelectedAnim.Get(); // GetAnimAsset();
 		ContinuingProperties.PlayingAssetAccumulatedTime = GetAccumulatedTime();
 
 		UPoseSearchInteractionLibrary::MotionMatchInteraction(CurrentResult, Availabilities, Context.AnimInstanceProxy->GetAnimInstanceObject(), ContinuingProperties, FName(),  & PoseHistoryProvider->GetPoseHistory(), bValidateResultAgainstAvailabilities);
 		check(CurrentResult.ActorRootTransforms.Num() == CurrentResult.ActorRootBoneTransforms.Num());
 
-		UAnimationAsset* RoledAnimAsset = nullptr;
-		if (const UMultiAnimAsset* MultiAnimAsset = Cast<UMultiAnimAsset>(CurrentResult.SelectedAnimation))
+		const UAnimationAsset* RoledAnimAsset = nullptr;
+		if (const UMultiAnimAsset* MultiAnimAsset = Cast<UMultiAnimAsset>(CurrentResult.SelectedAnim))
 		{
 			RoledAnimAsset = MultiAnimAsset->GetAnimationAsset(CurrentResult.Role);
 		}
-		else if (UAnimationAsset* SingleAnimAsset = Cast<UAnimationAsset>(CurrentResult.SelectedAnimation))
+		else if (const UAnimationAsset* SingleAnimAsset = Cast<UAnimationAsset>(CurrentResult.SelectedAnim))
 		{
 			RoledAnimAsset = SingleAnimAsset;
 		}
@@ -86,7 +86,7 @@ void FAnimNode_MotionMatchingInteraction::UpdateAssetPlayer(const FAnimationUpda
 			const FPoseSearchRoledSkeleton* RoledSkeleton = CurrentResult.SelectedDatabase->Schema->GetRoledSkeleton(CurrentResult.Role);
 			check(RoledSkeleton);
 
-			BlendTo(Context, RoledAnimAsset, CurrentResult.SelectedTime, CurrentResult.bLoop, CurrentResult.bIsMirrored, RoledSkeleton->MirrorDataTable.Get(),
+			BlendTo(Context, const_cast<UAnimationAsset*>(RoledAnimAsset), CurrentResult.SelectedTime, CurrentResult.bLoop, CurrentResult.bIsMirrored, RoledSkeleton->MirrorDataTable.Get(),
 				BlendTime, BlendProfile, BlendOption, bUseInertialBlend, NAME_None, CurrentResult.BlendParameters, CurrentResult.WantedPlayRate);
 
 			bBlendToExecuted = true;
@@ -139,7 +139,7 @@ void FAnimNode_MotionMatchingInteraction::Evaluate_AnyThread(FPoseContext& Outpu
 		return;
 	}
 
-	const UMultiAnimAsset* MultiAnimAsset = Cast<UMultiAnimAsset>(CurrentResult.SelectedAnimation);
+	const UMultiAnimAsset* MultiAnimAsset = Cast<UMultiAnimAsset>(CurrentResult.SelectedAnim);
 	if (!MultiAnimAsset)
 	{
 		return;
diff --git a/Engine/Plugins/Animation/PoseSearch/Source/Runtime/Private/MotionMatchingAnimNodeLibrary.cpp b/Engine/Plugins/Animation/PoseSearch/Source/Runtime/Private/MotionMatchingAnimNodeLibrary.cpp
--- a/Engine/Plugins/Animation/PoseSearch/Source/Runtime/Private/MotionMatchingAnimNodeLibrary.cpp
+++ b/Engine/Plugins/Animation/PoseSearch/Source/Runtime/Private/MotionMatchingAnimNodeLibrary.cpp
@@ -1,9 +1,11 @@
 // Copyright Epic Games, Inc. All Rights Reserved.
 
 #include "PoseSearch/MotionMatchingAnimNodeLibrary.h"
+#include "AlphaBlend.h"
 #include "PoseSearch/AnimNode_MotionMatching.h"
 #include "PoseSearch/PoseSearchDatabase.h"
-#include "AlphaBlend.h"
+#include "PoseSearch/PoseSearchDerivedData.h"
+#include "PoseSearch/PoseSearchSchema.h"
 
 FMotionMatchingBlueprintBlendSettings::FMotionMatchingBlueprintBlendSettings()
 	: BlendTime(0.2f)
@@ -21,34 +23,50 @@ FMotionMatchingAnimNodeReference UMotionMatchingAnimNodeLibrary::ConvertToMotion
 void UMotionMatchingAnimNodeLibrary::GetMotionMatchingSearchResult(const FMotionMatchingAnimNodeReference& MotionMatchingNode, FPoseSearchBlueprintResult& Result, bool& bIsResultValid)
 {
 	using namespace UE::PoseSearch;
-
+	
+	Result = FPoseSearchBlueprintResult();
 	bIsResultValid = false;
+
 	if (FAnimNode_MotionMatching* MotionMatchingNodePtr = MotionMatchingNode.GetAnimNodePtr<FAnimNode_MotionMatching>())
 	{
 		const FMotionMatchingState& MotionMatchingState = MotionMatchingNodePtr->GetMotionMatchingState();
-		if (const FSearchIndexAsset* SearchIndexAsset = MotionMatchingState.CurrentSearchResult.GetSearchIndexAsset())
+		const UPoseSearchDatabase* CurrentResultDatabase = MotionMatchingState.CurrentSearchResult.Database.Get();
+		if (CurrentResultDatabase == nullptr || CurrentResultDatabase->Schema == nullptr)
+		{
+			// @todo: should we log a warning?
+		}
+#if WITH_EDITOR
+		else if (EAsyncBuildIndexResult::Success != FAsyncPoseSearchDatabasesManagement::RequestAsyncBuildIndex(CurrentResultDatabase, ERequestAsyncBuildFlag::ContinueRequest))
+		{
+			UE_LOG(LogPoseSearch, Warning, TEXT("UMotionMatchingAnimNodeLibrary::GetMotionMatchingSearchResult called while '%s' is indexing. returning an invalid result.."), *CurrentResultDatabase->GetName());
+		}
+		// in case this node is not updated, and MotionMatchingState.CurrentSearchResult.Database gets modified, we could end up with CurrentSearchResult being out of synch with the updated database, so we need to reset the state
+		else if (MotionMatchingState.CurrentSearchResult.IsValid() && MotionMatchingState.CurrentSearchResult.PoseIdx >= MotionMatchingState.CurrentSearchResult.Database->GetSearchIndex().GetNumPoses())
+		{
+			UE_LOG(LogPoseSearch, Warning, TEXT("UMotionMatchingAnimNodeLibrary::GetMotionMatchingSearchResult inconsistency between MotionMatchingState and '%s', probably because of stale MotionMatchingState after database got reindexed. returning an invalid result.."), *CurrentResultDatabase->GetName());
+		}
+#endif // WITH_EDITOR
+		else if (const FSearchIndexAsset* SearchIndexAsset = MotionMatchingState.CurrentSearchResult.GetSearchIndexAsset())
 		{
-			const UPoseSearchDatabase* CurrentResultDatabase = MotionMatchingState.CurrentSearchResult.Database.Get();
-			if (CurrentResultDatabase && CurrentResultDatabase->Schema)
+			if (const FPoseSearchDatabaseAnimationAssetBase* DatabaseAsset = CurrentResultDatabase->GetDatabaseAnimationAsset<FPoseSearchDatabaseAnimationAssetBase>(*SearchIndexAsset))
 			{
-				if (const FPoseSearchDatabaseAnimationAssetBase* DatabaseAsset = CurrentResultDatabase->GetDatabaseAnimationAsset<FPoseSearchDatabaseAnimationAssetBase>(*SearchIndexAsset))
-				{
-					Result.SelectedAnimation = DatabaseAsset->GetAnimationAsset();
-					Result.SelectedTime = MotionMatchingState.CurrentSearchResult.AssetTime;
-					Result.bIsContinuingPoseSearch = MotionMatchingState.CurrentSearchResult.bIsContinuingPoseSearch;
-					Result.bLoop = SearchIndexAsset->IsLooping();
-					Result.bIsMirrored = SearchIndexAsset->IsMirrored();
-					Result.BlendParameters = SearchIndexAsset->GetBlendParameters();
-					Result.SelectedDatabase = CurrentResultDatabase;
-					Result.WantedPlayRate = MotionMatchingState.WantedPlayRate;
-					bIsResultValid = true;
-				}
+				Result.SelectedAnim = DatabaseAsset->GetAnimationAsset();
+				Result.SelectedTime = MotionMatchingState.CurrentSearchResult.AssetTime;
+				Result.bIsContinuingPoseSearch = MotionMatchingState.CurrentSearchResult.bIsContinuingPoseSearch;
+				Result.WantedPlayRate = MotionMatchingState.WantedPlayRate;
+				Result.bLoop = SearchIndexAsset->IsLooping();
+				Result.bIsMirrored = SearchIndexAsset->IsMirrored();
+				Result.BlendParameters = SearchIndexAsset->GetBlendParameters();
+				Result.SelectedDatabase = CurrentResultDatabase;
+				Result.SearchCost = MotionMatchingState.CurrentSearchResult.PoseCost;
+				Result.Role = CurrentResultDatabase->Schema->GetDefaultRole();
+				bIsResultValid = true;
 			}
 		}
 	}
 	else
 	{
-		UE_LOG(LogPoseSearch, Warning, TEXT("UMotionMatchingAnimNodeLibrary::SetDatabase called on an invalid context or with an invalid type"));
+		UE_LOG(LogPoseSearch, Warning, TEXT("UMotionMatchingAnimNodeLibrary::GetMotionMatchingSearchResult called on an invalid context or with an invalid type"));
 	}
 }
 
@@ -91,7 +109,7 @@ void UMotionMatchingAnimNodeLibrary::SetDatabaseToSearch(const FMotionMatchingAn
 	}
 	else
 	{
-		UE_LOG(LogPoseSearch, Warning, TEXT("UMotionMatchingAnimNodeLibrary::SetDatabase called on an invalid context or with an invalid type"));
+		UE_LOG(LogPoseSearch, Warning, TEXT("UMotionMatchingAnimNodeLibrary::SetDatabaseToSearch called on an invalid context or with an invalid type"));
 	}
 }
 
@@ -103,7 +121,7 @@ void UMotionMatchingAnimNodeLibrary::SetDatabasesToSearch(const FMotionMatchingA
 	}
 	else
 	{
-		UE_LOG(LogPoseSearch, Warning, TEXT("UMotionMatchingAnimNodeLibrary::SetDatabases called on an invalid context or with an invalid type"));
+		UE_LOG(LogPoseSearch, Warning, TEXT("UMotionMatchingAnimNodeLibrary::SetDatabasesToSearch called on an invalid context or with an invalid type"));
 	}
 }
 
@@ -127,6 +145,6 @@ void UMotionMatchingAnimNodeLibrary::SetInterruptMode(const FMotionMatchingAnimN
 	}
 	else
 	{
-		UE_LOG(LogPoseSearch, Warning, TEXT("UMotionMatchingAnimNodeLibrary::ForceInterruptOneFrame called on an invalid context or with an invalid type"));
+		UE_LOG(LogPoseSearch, Warning, TEXT("UMotionMatchingAnimNodeLibrary::SetInterruptMode called on an invalid context or with an invalid type"));
 	}
 }
\ No newline at end of file
diff --git a/Engine/Plugins/Animation/PoseSearch/Source/Runtime/Private/PoseSearchInteractionIsland.cpp b/Engine/Plugins/Animation/PoseSearch/Source/Runtime/Private/PoseSearchInteractionIsland.cpp
--- a/Engine/Plugins/Animation/PoseSearch/Source/Runtime/Private/PoseSearchInteractionIsland.cpp
+++ b/Engine/Plugins/Animation/PoseSearch/Source/Runtime/Private/PoseSearchInteractionIsland.cpp
@@ -971,7 +971,7 @@ void FInteractionIsland::LogTickDependencies() const
 }
 #endif // ENABLE_ANIM_DEBUG
 
-bool FInteractionIsland::DoSearch_AnyThread(const UObject* AnimContext, const FPoseSearchContinuingProperties& ContinuingProperties, FPoseSearchInteractionBlueprintResult& Result)
+bool FInteractionIsland::DoSearch_AnyThread(const UObject* AnimContext, const FPoseSearchContinuingProperties& ContinuingProperties, FPoseSearchBlueprintResult& Result)
 {
 	check(AnimContext);
 
@@ -987,7 +987,7 @@ bool FInteractionIsland::DoSearch_AnyThread(const UObject* AnimContext, const FP
 	{
 		// search failed, because AnimContext is not the context from the main actor. continuing the search would lead to threading issues,
 		// since the threading model expect the main actor to perform the search! (this happens when recompiling animation blue prints..)
-		Result = FPoseSearchInteractionBlueprintResult();
+		Result = FPoseSearchBlueprintResult();
 		return false;
 	}
 
@@ -1059,7 +1059,7 @@ bool FInteractionIsland::DoSearch_AnyThread(const UObject* AnimContext, const FP
 	return GetResult_AnyThread(AnimContext, Result);
 }
 
-bool FInteractionIsland::GetResult_AnyThread(const UObject* AnimContext, FPoseSearchInteractionBlueprintResult& Result, bool bCompareOwningActors)
+bool FInteractionIsland::GetResult_AnyThread(const UObject* AnimContext, FPoseSearchBlueprintResult& Result, bool bCompareOwningActors)
 {
 	check(AnimContext);
 
@@ -1094,7 +1094,7 @@ bool FInteractionIsland::GetResult_AnyThread(const UObject* AnimContext, FPoseSe
 
 				const FRole Role = SearchContext.Roles[AnimContextIndex];
 
-				Result.SelectedAnimation = DatabaseAnimationAssetBase->GetAnimationAsset();
+				Result.SelectedAnim = DatabaseAnimationAssetBase->GetAnimationAsset();
 				Result.SelectedTime = SearchResult.AssetTime;
 				Result.bIsContinuingPoseSearch = SearchResult.bIsContinuingPoseSearch;
 				Result.bLoop = SearchIndexAsset->IsLooping();
@@ -1102,6 +1102,7 @@ bool FInteractionIsland::GetResult_AnyThread(const UObject* AnimContext, FPoseSe
 				Result.BlendParameters = SearchIndexAsset->GetBlendParameters();
 				Result.SelectedDatabase = Database;
 				Result.SearchCost = SearchResult.PoseCost;
+				Result.bIsInteraction = true;
 				Result.Role = Role;
 
 				// figuring out the WantedPlayRate
@@ -1120,7 +1121,7 @@ bool FInteractionIsland::GetResult_AnyThread(const UObject* AnimContext, FPoseSe
 				//	}
 				//}
 
-				if (const UMultiAnimAsset* MultiAnimAsset = Cast<UMultiAnimAsset>(Result.SelectedAnimation))
+				if (const UMultiAnimAsset* MultiAnimAsset = Cast<UMultiAnimAsset>(Result.SelectedAnim))
 				{
 					const int32 NumRoles = MultiAnimAsset->GetNumRoles();
 					Result.ActorRootTransforms.SetNum(NumRoles);
@@ -1161,7 +1162,7 @@ bool FInteractionIsland::GetResult_AnyThread(const UObject* AnimContext, FPoseSe
 		}
 	}
 
-	Result = FPoseSearchInteractionBlueprintResult();
+	Result = FPoseSearchBlueprintResult();
 	return false;
 }
 
diff --git a/Engine/Plugins/Animation/PoseSearch/Source/Runtime/Private/PoseSearchInteractionLibrary.cpp b/Engine/Plugins/Animation/PoseSearch/Source/Runtime/Private/PoseSearchInteractionLibrary.cpp
--- a/Engine/Plugins/Animation/PoseSearch/Source/Runtime/Private/PoseSearchInteractionLibrary.cpp
+++ b/Engine/Plugins/Animation/PoseSearch/Source/Runtime/Private/PoseSearchInteractionLibrary.cpp
@@ -6,21 +6,21 @@
 #include "PoseSearch/PoseSearchInteractionSubsystem.h"
 #include "Engine/World.h"
 
-FPoseSearchInteractionBlueprintResult UPoseSearchInteractionLibrary::MotionMatchInteraction_Pure(TArray<FPoseSearchInteractionAvailability> Availabilities, const UObject* AnimContext, FPoseSearchContinuingProperties ContinuingProperties, FName PoseHistoryName, bool bValidateResultAgainstAvailabilities)
+FPoseSearchBlueprintResult UPoseSearchInteractionLibrary::MotionMatchInteraction_Pure(TArray<FPoseSearchInteractionAvailability> Availabilities, const UObject* AnimContext, FPoseSearchContinuingProperties ContinuingProperties, FName PoseHistoryName, bool bValidateResultAgainstAvailabilities)
 {
-	FPoseSearchInteractionBlueprintResult Result;
+	FPoseSearchBlueprintResult Result;
 	MotionMatchInteraction(Result, Availabilities, AnimContext, ContinuingProperties, PoseHistoryName, nullptr, bValidateResultAgainstAvailabilities);
 	return Result;
 }
 
-FPoseSearchInteractionBlueprintResult UPoseSearchInteractionLibrary::MotionMatchInteraction(TArray<FPoseSearchInteractionAvailability> Availabilities, const UObject* AnimContext, FPoseSearchContinuingProperties ContinuingProperties, FName PoseHistoryName, bool bValidateResultAgainstAvailabilities)
+FPoseSearchBlueprintResult UPoseSearchInteractionLibrary::MotionMatchInteraction(TArray<FPoseSearchInteractionAvailability> Availabilities, const UObject* AnimContext, FPoseSearchContinuingProperties ContinuingProperties, FName PoseHistoryName, bool bValidateResultAgainstAvailabilities)
 {
-	FPoseSearchInteractionBlueprintResult Result;
+	FPoseSearchBlueprintResult Result;
 	MotionMatchInteraction(Result, Availabilities, AnimContext, ContinuingProperties, PoseHistoryName, nullptr, bValidateResultAgainstAvailabilities);
 	return Result;
 }
 
-void UPoseSearchInteractionLibrary::MotionMatchInteraction(FPoseSearchInteractionBlueprintResult& Result, const TArrayView<const FPoseSearchInteractionAvailability> Availabilities, const UObject* AnimContext, const FPoseSearchContinuingProperties& ContinuingProperties, FName PoseHistoryName, const UE::PoseSearch::IPoseHistory* PoseHistory, bool bValidateResultAgainstAvailabilities)
+void UPoseSearchInteractionLibrary::MotionMatchInteraction(FPoseSearchBlueprintResult& Result, const TArrayView<const FPoseSearchInteractionAvailability> Availabilities, const UObject* AnimContext, const FPoseSearchContinuingProperties& ContinuingProperties, FName PoseHistoryName, const UE::PoseSearch::IPoseHistory* PoseHistory, bool bValidateResultAgainstAvailabilities)
 {
 	if (UPoseSearchInteractionSubsystem* InteractionSubsystem = UPoseSearchInteractionSubsystem::GetSubsystem_AnyThread(AnimContext))
 	{
@@ -28,13 +28,13 @@ void UPoseSearchInteractionLibrary::MotionMatchInteraction(FPoseSearchInteractio
 	}
 	else
 	{
-		Result = FPoseSearchInteractionBlueprintResult();
+		Result = FPoseSearchBlueprintResult();
 	}
 }
 
-FPoseSearchInteractionBlueprintResult UPoseSearchInteractionLibrary::GetCachedInteraction(const UObject* AnimContext, bool bCompareOwningActors)
+FPoseSearchBlueprintResult UPoseSearchInteractionLibrary::GetCachedInteraction(const UObject* AnimContext, bool bCompareOwningActors)
 {
-	FPoseSearchInteractionBlueprintResult Result;
+	FPoseSearchBlueprintResult Result;
 	if (UPoseSearchInteractionSubsystem* InteractionSubsystem = UPoseSearchInteractionSubsystem::GetSubsystem_AnyThread(AnimContext))
 	{
 		InteractionSubsystem->GetResult_AnyThread(AnimContext, Result, bCompareOwningActors);
diff --git a/Engine/Plugins/Animation/PoseSearch/Source/Runtime/Private/PoseSearchInteractionSubsystem.cpp b/Engine/Plugins/Animation/PoseSearch/Source/Runtime/Private/PoseSearchInteractionSubsystem.cpp
--- a/Engine/Plugins/Animation/PoseSearch/Source/Runtime/Private/PoseSearchInteractionSubsystem.cpp
+++ b/Engine/Plugins/Animation/PoseSearch/Source/Runtime/Private/PoseSearchInteractionSubsystem.cpp
@@ -1160,14 +1160,14 @@ TStatId UPoseSearchInteractionSubsystem::GetStatId() const
 }
 
 void UPoseSearchInteractionSubsystem::Query_AnyThread(const TArrayView<const FPoseSearchInteractionAvailability> Availabilities, const UObject* AnimContext, 
-	const FPoseSearchContinuingProperties& ContinuingProperties, FPoseSearchInteractionBlueprintResult& Result,
+	const FPoseSearchContinuingProperties& ContinuingProperties, FPoseSearchBlueprintResult& Result,
 	FName PoseHistoryName, const UE::PoseSearch::IPoseHistory* PoseHistory, bool bValidateResultAgainstAvailabilities)
 {
 	QUICK_SCOPE_CYCLE_COUNTER(STAT_UPoseSearchInteractionSubsystem_Query_AnyThread);
 
 	using namespace UE::PoseSearch;
 
-	Result = FPoseSearchInteractionBlueprintResult();
+	Result = FPoseSearchBlueprintResult();
 
 #if !NO_CVARS
 	if (!GVarPoseSearchInteractionEnabled)
@@ -1181,7 +1181,7 @@ void UPoseSearchInteractionSubsystem::Query_AnyThread(const TArrayView<const FPo
 	{
 		Island->DoSearch_AnyThread(AnimContext, ContinuingProperties, Result);
 
-		if (bValidateResultAgainstAvailabilities && Result.SelectedAnimation)
+		if (bValidateResultAgainstAvailabilities && Result.SelectedAnim != nullptr)
 		{
 			bool bResultValidated = false;
 
@@ -1197,7 +1197,7 @@ void UPoseSearchInteractionSubsystem::Query_AnyThread(const TArrayView<const FPo
 
 			if (!bResultValidated)
 			{
-				Result = FPoseSearchInteractionBlueprintResult();
+				Result = FPoseSearchBlueprintResult();
 			}
 		}
 	}
@@ -1209,7 +1209,7 @@ void UPoseSearchInteractionSubsystem::Query_AnyThread(const TArrayView<const FPo
 	}
 }
 
-void UPoseSearchInteractionSubsystem::GetResult_AnyThread(const UObject* AnimContext, FPoseSearchInteractionBlueprintResult& Result, bool bCompareOwningActors)
+void UPoseSearchInteractionSubsystem::GetResult_AnyThread(const UObject* AnimContext, FPoseSearchBlueprintResult& Result, bool bCompareOwningActors)
 {
 	using namespace UE::PoseSearch;
 
@@ -1219,7 +1219,7 @@ void UPoseSearchInteractionSubsystem::GetResult_AnyThread(const UObject* AnimCon
 	}
 	else
 	{
-		Result = FPoseSearchInteractionBlueprintResult();
+		Result = FPoseSearchBlueprintResult();
 	}
 }
 
diff --git a/Engine/Plugins/Animation/PoseSearch/Source/Runtime/Private/PoseSearchInteractionUtils.cpp b/Engine/Plugins/Animation/PoseSearch/Source/Runtime/Private/PoseSearchInteractionUtils.cpp
--- a/Engine/Plugins/Animation/PoseSearch/Source/Runtime/Private/PoseSearchInteractionUtils.cpp
+++ b/Engine/Plugins/Animation/PoseSearch/Source/Runtime/Private/PoseSearchInteractionUtils.cpp
@@ -37,9 +37,9 @@ FRoleToIndex MakeRoleToIndex(const UMultiAnimAsset* MultiAnimAsset)
 	return RoleToIndex;
 }
 
-void CalculateFullAlignedTransforms(const FPoseSearchInteractionBlueprintResult& CurrentResult, bool bWarpUsingRootBone, TArrayView<FTransform> OutFullAlignedTransforms)
+void CalculateFullAlignedTransforms(const FPoseSearchBlueprintResult& CurrentResult, bool bWarpUsingRootBone, TArrayView<FTransform> OutFullAlignedTransforms)
 {
-	const UMultiAnimAsset* MultiAnimAsset = CastChecked<UMultiAnimAsset>(CurrentResult.SelectedAnimation);
+	const UMultiAnimAsset* MultiAnimAsset = CastChecked<UMultiAnimAsset>(CurrentResult.SelectedAnim);
 	const int32 NumRoles = MultiAnimAsset->GetNumRoles();
 	
 	check(OutFullAlignedTransforms.Num() == NumRoles);
diff --git a/Engine/Plugins/Animation/PoseSearch/Source/Runtime/Private/PoseSearchLibrary.cpp b/Engine/Plugins/Animation/PoseSearch/Source/Runtime/Private/PoseSearchLibrary.cpp
--- a/Engine/Plugins/Animation/PoseSearch/Source/Runtime/Private/PoseSearchLibrary.cpp
+++ b/Engine/Plugins/Animation/PoseSearch/Source/Runtime/Private/PoseSearchLibrary.cpp
@@ -886,14 +886,7 @@ void UPoseSearchLibrary::MotionMatch(
 	using namespace UE::Anim;
 	using namespace UE::PoseSearch;
 
-	Result.SelectedAnimation = nullptr;
-	Result.SelectedTime = 0.f;
-	Result.bIsContinuingPoseSearch = false;
-	Result.bLoop = false;
-	Result.bIsMirrored = false;
-	Result.BlendParameters = FVector::ZeroVector;
-	Result.SelectedDatabase = nullptr;
-	Result.SearchCost = MAX_flt;
+	Result = FPoseSearchBlueprintResult();
 
 	if (AnimInstances.IsEmpty() || AnimInstances.Num() != Roles.Num())
 	{
@@ -943,7 +936,7 @@ void UPoseSearchLibrary::MotionMatch(
 		const FSearchIndexAsset* SearchIndexAsset = SearchResult.GetSearchIndexAsset();
 		if (const FPoseSearchDatabaseAnimationAssetBase* DatabaseAsset = Database->GetDatabaseAnimationAsset<FPoseSearchDatabaseAnimationAssetBase>(*SearchIndexAsset))
 		{
-			Result.SelectedAnimation = DatabaseAsset->GetAnimationAsset();
+			Result.SelectedAnim = DatabaseAsset->GetAnimationAsset();
 			Result.SelectedTime = SearchResult.AssetTime;
 			Result.bIsContinuingPoseSearch = SearchResult.bIsContinuingPoseSearch;
 			Result.bLoop = SearchIndexAsset->IsLooping();
diff --git a/Engine/Plugins/Animation/PoseSearch/Source/Runtime/Public/PoseSearch/AnimNode_MotionMatchingInteraction.h b/Engine/Plugins/Animation/PoseSearch/Source/Runtime/Public/PoseSearch/AnimNode_MotionMatchingInteraction.h
--- a/Engine/Plugins/Animation/PoseSearch/Source/Runtime/Public/PoseSearch/AnimNode_MotionMatchingInteraction.h
+++ b/Engine/Plugins/Animation/PoseSearch/Source/Runtime/Public/PoseSearch/AnimNode_MotionMatchingInteraction.h
@@ -62,12 +62,12 @@ struct POSESEARCH_API FAnimNode_MotionMatchingInteraction : public FAnimNode_Ble
 
 	virtual void Reset() override;
 
-	bool IsInteracting() const { return CurrentResult.SelectedAnimation != nullptr; }
+	bool IsInteracting() const { return CurrentResult.bIsInteraction; }
 
 protected:
 	bool NeedsReset(const FAnimationUpdateContext& Context) const;
 
-	FPoseSearchInteractionBlueprintResult CurrentResult;
+	FPoseSearchBlueprintResult CurrentResult;
 	FTransform MeshWithOffset = FTransform::Identity;
 	FTransform MeshWithoutOffset = FTransform::Identity;
 
diff --git a/Engine/Plugins/Animation/PoseSearch/Source/Runtime/Public/PoseSearch/PoseSearchInteractionIsland.h b/Engine/Plugins/Animation/PoseSearch/Source/Runtime/Public/PoseSearch/PoseSearchInteractionIsland.h
--- a/Engine/Plugins/Animation/PoseSearch/Source/Runtime/Public/PoseSearch/PoseSearchInteractionIsland.h
+++ b/Engine/Plugins/Animation/PoseSearch/Source/Runtime/Public/PoseSearch/PoseSearchInteractionIsland.h
@@ -5,7 +5,7 @@
 #include "PoseSearch/PoseSearchInteractionValidator.h"
 #include "PoseSearch/PoseSearchLibrary.h"
 
-struct FPoseSearchInteractionBlueprintResult;
+struct FPoseSearchBlueprintResult;
 struct FPoseSearchContinuingProperties;
 class UMultiAnimAsset;
 class UPoseSearchInteractionSubsystem;
@@ -89,8 +89,8 @@ struct FInteractionIsland
 	FInteractionIsland(ULevel* Level, UPoseSearchInteractionSubsystem* Subsystem);
 	~FInteractionIsland();
 
-	bool DoSearch_AnyThread(const UObject* AnimContext, const FPoseSearchContinuingProperties& ContinuingProperties, FPoseSearchInteractionBlueprintResult& Result);
-	bool GetResult_AnyThread(const UObject* AnimContext, FPoseSearchInteractionBlueprintResult& Result, bool bCompareOwningActors = false);
+	bool DoSearch_AnyThread(const UObject* AnimContext, const FPoseSearchContinuingProperties& ContinuingProperties, FPoseSearchBlueprintResult& Result);
+	bool GetResult_AnyThread(const UObject* AnimContext, FPoseSearchBlueprintResult& Result, bool bCompareOwningActors = false);
 
 	const TArray<TWeakObjectPtr<UActorComponent>>& GetTickActorComponents() const { return TickActorComponents; }
 	const TArray<TWeakObjectPtr<const UObject>>& GetIslandAnimContexts() const { return IslandAnimContexts; }
diff --git a/Engine/Plugins/Animation/PoseSearch/Source/Runtime/Public/PoseSearch/PoseSearchInteractionLibrary.h b/Engine/Plugins/Animation/PoseSearch/Source/Runtime/Public/PoseSearch/PoseSearchInteractionLibrary.h
--- a/Engine/Plugins/Animation/PoseSearch/Source/Runtime/Public/PoseSearch/PoseSearchInteractionLibrary.h
+++ b/Engine/Plugins/Animation/PoseSearch/Source/Runtime/Public/PoseSearch/PoseSearchInteractionLibrary.h
@@ -61,66 +61,6 @@ struct POSESEARCH_API FPoseSearchInteractionAvailability
 	}
 };
 
-// Experimental, this feature might be removed without warning, not for production use
-USTRUCT(Experimental, BlueprintType, Category="Animation|Pose Search")
-struct POSESEARCH_API FPoseSearchInteractionBlueprintResult
-{
-	GENERATED_BODY()
-public:
-
-	// animation assigned to this character to partecipate in the interaction
-	UPROPERTY(Transient, VisibleAnywhere, BlueprintReadOnly, Category=State)
-	TObjectPtr<UObject> SelectedAnimation = nullptr;
-	
-	// SelectedAnimation associated time
-	UPROPERTY(Transient, VisibleAnywhere, BlueprintReadOnly, Category=State)
-	float SelectedTime = 0.f;
-	
-	// SelectedAnimation at SelectedTime is from the continuing pose search
-	UPROPERTY(Transient, VisibleAnywhere, BlueprintReadOnly, Category=State)
-	bool bIsContinuingPoseSearch = false;
-
-	// SelectedAnimation associated play rate
-	UPROPERTY(Transient, VisibleAnywhere, BlueprintReadOnly, Category=State)
-	float WantedPlayRate = 0.f;
-
-	// SelectedAnimation associated looping state
-	UPROPERTY(Transient, VisibleAnywhere, BlueprintReadOnly, Category=State)
-	bool bLoop = false;
-	
-	// SelectedAnimation associated mirror state
-	UPROPERTY(Transient, VisibleAnywhere, BlueprintReadOnly, Category=State)
-	bool bIsMirrored = false;
-	
-	// SelectedAnimation associated BlendParameters (if SelectedAnimation is a UBlendSpace)
-	UPROPERTY(Transient, VisibleAnywhere, BlueprintReadOnly, Category=State)
-	FVector BlendParameters = FVector::ZeroVector;
-
-	// selected SelectedDatabase for this character interaction
-	UPROPERTY(Transient, VisibleAnywhere, BlueprintReadOnly, Category=State)
-	TWeakObjectPtr<const UPoseSearchDatabase> SelectedDatabase = nullptr;
-
-	// associated motion matching search cost for this result
-	UPROPERTY(Transient, VisibleAnywhere, BlueprintReadOnly, Category=State)
-	float SearchCost = MAX_flt;
-
-	// assigned role to this character (AnimInstance)
-	UPROPERTY(Transient, VisibleAnywhere, BlueprintReadOnly, Category=State)
-	FName Role;
-
-	// cached actors root transforms for all the roles in SelectedAnimation (as UMultiAnimAsset). If the role is missing because optional the ActorRootTransform will be identity
-	UPROPERTY(Transient, VisibleAnywhere, BlueprintReadOnly, Category=State)
-	TArray<FTransform> ActorRootTransforms;
-
-	// cached actors root bone transforms for all the roles in SelectedAnimation (as UMultiAnimAsset). If the role is missing because optional the ActorRootBoneTransform will be identity
-	UPROPERTY(Transient, VisibleAnywhere, BlueprintReadOnly, Category=State)
-	TArray<FTransform> ActorRootBoneTransforms;
-
-	// AnimContexts involved in this interaction for all the roles in SelectedAnimation (as UMultiAnimAsset). If the role is missing because optional the AnimContext will be null
-	// BE CAREFUL HOW YOU USE THEM! since accessing data from other actors without proper tick dependencies leads to threading problems!
-	TArray<TWeakObjectPtr<const UObject>> AnimContexts;
-};
-
 // Experimental, this feature might be removed without warning, not for production use
 UCLASS(Experimental)
 class POSESEARCH_API UPoseSearchInteractionLibrary : public UBlueprintFunctionLibrary
@@ -133,22 +73,22 @@ public:
 	// FPoseSearchInteractionAvailability represents the character availability to partecipate in an interaction with other characters for the next frame.
 	// that means there will always be one frame delay between publiching availabilities and getting a result back from MotionMatchInteraction_Pure!
 	// 
-	// if FPoseSearchInteractionBlueprintResult has a valid SelectedAnimation, this will be the animation assigned to this character to partecipate in this interaction.
+	// if FPoseSearchBlueprintResult has a valid SelectedAnimation, this will be the animation assigned to this character to partecipate in this interaction.
 	// additional interaction properties, like assigned role, SelectedAnimation time, SearchCost, etc can be found within the result
 	// ContinuingProperties are used to figure out the continuing pose and bias it accordingly. ContinuingProperties can reference directly the UMultiAnimAsset
 	// or any of the roled UMultiAnimAsset::GetAnimationAsset, and the UPoseSearchInteractionSubsystem will figure out the related UMultiAnimAsset
 	// PoseHistoryName is the name of the pose history node used for the associated motion matching search
 	// if bValidateResultAgainstAvailabilities is true, the result will be invalidated if doesn't respect the new availabilities
 	UFUNCTION(Experimental, BlueprintPure, Category = "Animation|Pose Search", meta = (BlueprintThreadSafe))
-	static FPoseSearchInteractionBlueprintResult MotionMatchInteraction_Pure(TArray<FPoseSearchInteractionAvailability> Availabilities, const UObject* AnimContext, FPoseSearchContinuingProperties ContinuingProperties, FName PoseHistoryName, bool bValidateResultAgainstAvailabilities = true);
+	static FPoseSearchBlueprintResult MotionMatchInteraction_Pure(TArray<FPoseSearchInteractionAvailability> Availabilities, const UObject* AnimContext, FPoseSearchContinuingProperties ContinuingProperties, FName PoseHistoryName, bool bValidateResultAgainstAvailabilities = true);
 
 	// BlueprintCallable version of MotionMatchInteraction_Pure
 	UFUNCTION(Experimental, BlueprintCallable, Category = "Animation|Pose Search", meta = (BlueprintThreadSafe))
-	static FPoseSearchInteractionBlueprintResult MotionMatchInteraction(TArray<FPoseSearchInteractionAvailability> Availabilities, const UObject* AnimContext, FPoseSearchContinuingProperties ContinuingProperties, FName PoseHistoryName, bool bValidateResultAgainstAvailabilities = true);
+	static FPoseSearchBlueprintResult MotionMatchInteraction(TArray<FPoseSearchInteractionAvailability> Availabilities, const UObject* AnimContext, FPoseSearchContinuingProperties ContinuingProperties, FName PoseHistoryName, bool bValidateResultAgainstAvailabilities = true);
 
-	static void MotionMatchInteraction(FPoseSearchInteractionBlueprintResult& Result, const TArrayView<const FPoseSearchInteractionAvailability> Availabilities, const UObject* AnimContext, const FPoseSearchContinuingProperties& ContinuingProperties, FName PoseHistoryName, const UE::PoseSearch::IPoseHistory* PoseHistory, bool bValidateResultAgainstAvailabilities = true);
+	static void MotionMatchInteraction(FPoseSearchBlueprintResult& Result, const TArrayView<const FPoseSearchInteractionAvailability> Availabilities, const UObject* AnimContext, const FPoseSearchContinuingProperties& ContinuingProperties, FName PoseHistoryName, const UE::PoseSearch::IPoseHistory* PoseHistory, bool bValidateResultAgainstAvailabilities = true);
 
-	static FPoseSearchInteractionBlueprintResult GetCachedInteraction(const UObject* AnimContext, bool bCompareOwningActors = false);
+	static FPoseSearchBlueprintResult GetCachedInteraction(const UObject* AnimContext, bool bCompareOwningActors = false);
 	
 	UFUNCTION(Experimental, BlueprintPure, Category = "Animation|Pose Search", meta = (BlueprintThreadSafe))
 	static FPoseSearchContinuingProperties GetMontageContinuingProperties(UAnimInstance* AnimInstance);
diff --git a/Engine/Plugins/Animation/PoseSearch/Source/Runtime/Public/PoseSearch/PoseSearchInteractionSubsystem.h b/Engine/Plugins/Animation/PoseSearch/Source/Runtime/Public/PoseSearch/PoseSearchInteractionSubsystem.h
--- a/Engine/Plugins/Animation/PoseSearch/Source/Runtime/Public/PoseSearch/PoseSearchInteractionSubsystem.h
+++ b/Engine/Plugins/Animation/PoseSearch/Source/Runtime/Public/PoseSearch/PoseSearchInteractionSubsystem.h
@@ -112,7 +112,7 @@ struct FTagToDatabases;
 // World subsystem accepting the publication of characters (via their AnimInstance(s)) FPoseSearchInteractionAvailability, representing the characters willingness to partecipate in an 
 // interaction with other characters from the next frame forward via Query_AnyThread method.
 //
-// The same method will return the FPoseSearchInteractionBlueprintResult from the PREVIOUS Tick processing (categorization of FPoseSearchInteractionAvailability(s) in multiple FInteractionIsland(s)),
+// The same method will return the FPoseSearchBlueprintResult from the PREVIOUS Tick processing (categorization of FPoseSearchInteractionAvailability(s) in multiple FInteractionIsland(s)),
 // to the requesting character, containing the animation to play at what time, and the assigned role to partecipate in the selected interaction within the assigned 
 // FInteractionIsland
 
@@ -167,16 +167,16 @@ public:
 	virtual TStatId GetStatId() const override;
 
 	// publishing FPoseSearchInteractionAvailability(s) for the requesting character (via AnimContext as UAnimInstance or UAnimNextComponent) 
-	// and getting the FPoseSearchInteractionBlueprintResult from the PREVIOUS Tick update
+	// and getting the FPoseSearchBlueprintResult from the PREVIOUS Tick update
 	// containing the animation to play at what time, and the assigned role to partecipate in the selected interaction.
 	// Either a PoseHistoryName or a PoseHistory are required to perform the associated motion matching searches
 	// if bValidateResultAgainstAvailabilities is true, the result will be validated against the currently published availabilities,
 	// that could differ from the previous frame one for this AnimContext, and discarded if not valid within the context of the new availabilities
 	void Query_AnyThread(const TArrayView<const FPoseSearchInteractionAvailability> Availabilities, const UObject* AnimContext, 
-		const FPoseSearchContinuingProperties& ContinuingProperties, FPoseSearchInteractionBlueprintResult& Result,
+		const FPoseSearchContinuingProperties& ContinuingProperties, FPoseSearchBlueprintResult& Result,
 		FName PoseHistoryName, const UE::PoseSearch::IPoseHistory* PoseHistory, bool bValidateResultAgainstAvailabilities);
 
-	void GetResult_AnyThread(const UObject* AnimContext, FPoseSearchInteractionBlueprintResult& Result, bool bCompareOwningActors = false);
+	void GetResult_AnyThread(const UObject* AnimContext, FPoseSearchBlueprintResult& Result, bool bCompareOwningActors = false);
 
 	// returns and array containing all the other AnimContext(s) that the input AnimContext is currently interacting with
 	TArray<const UObject*> GetInteractingAnimContextsWith(const UObject* AnimContext) const;
diff --git a/Engine/Plugins/Animation/PoseSearch/Source/Runtime/Public/PoseSearch/PoseSearchInteractionUtils.h b/Engine/Plugins/Animation/PoseSearch/Source/Runtime/Public/PoseSearch/PoseSearchInteractionUtils.h
--- a/Engine/Plugins/Animation/PoseSearch/Source/Runtime/Public/PoseSearch/PoseSearchInteractionUtils.h
+++ b/Engine/Plugins/Animation/PoseSearch/Source/Runtime/Public/PoseSearch/PoseSearchInteractionUtils.h
@@ -4,7 +4,7 @@
 
 #include "PoseSearch/PoseSearchRole.h"
 
-struct FPoseSearchInteractionBlueprintResult;
+struct FPoseSearchBlueprintResult;
 class UMultiAnimAsset;
 
 namespace UE::PoseSearch
@@ -17,7 +17,7 @@ POSESEARCH_API int32 GetRoleIndex(const UMultiAnimAsset* MultiAnimAsset, const F
 FRoleToIndex MakeRoleToIndex(const UMultiAnimAsset* MultiAnimAsset);
 
 // Experimental, this feature might be removed without warning, not for production use
-POSESEARCH_API void CalculateFullAlignedTransforms(const FPoseSearchInteractionBlueprintResult& CurrentResult, bool bWarpUsingRootBone, TArrayView<FTransform> OutFullAlignedTransforms);
+POSESEARCH_API void CalculateFullAlignedTransforms(const FPoseSearchBlueprintResult& CurrentResult, bool bWarpUsingRootBone, TArrayView<FTransform> OutFullAlignedTransforms);
 
 // Experimental, this feature might be removed without warning, not for production use
 POSESEARCH_API FTransform CalculateDeltaAlignment(const FTransform& MeshWithoutOffset, const FTransform& MeshWithOffset, const FTransform& FullAlignedTransform, float WarpingRotationRatio, float WarpingTranslationRatio);
diff --git a/Engine/Plugins/Animation/PoseSearch/Source/Runtime/Public/PoseSearch/PoseSearchResult.h b/Engine/Plugins/Animation/PoseSearch/Source/Runtime/Public/PoseSearch/PoseSearchResult.h
--- a/Engine/Plugins/Animation/PoseSearch/Source/Runtime/Public/PoseSearch/PoseSearchResult.h
+++ b/Engine/Plugins/Animation/PoseSearch/Source/Runtime/Public/PoseSearch/PoseSearchResult.h
@@ -102,31 +102,67 @@ struct POSESEARCH_API FPoseSearchBlueprintResult
 	GENERATED_BODY()
 public:
 
+#if WITH_EDITORONLY_DATA
+	// Deprecated in 5.6. Use "SelectedAnim" property instead
+	UPROPERTY(Transient, meta = (DeprecatedProperty, DeprecationMessage = "Use SelectedAnim instead"))
+	TObjectPtr<const UObject> SelectedAnimation_DEPRECATED = nullptr;
+#endif // WITH_EDITORONLY_DATA
+
+	// selected animation selected by the motion matching alghoritm that matches the query (trajectory + historical pose) the best
 	UPROPERTY(Transient, VisibleAnywhere, BlueprintReadOnly, Category=State)
-	TObjectPtr<const UObject> SelectedAnimation = nullptr;
+	TWeakObjectPtr<UObject> SelectedAnim = nullptr;
 	
+	// SelectedAnim associated time (to have the best result out of the motion matchign selection, play the SelectedAnim at the SelectedTime)
 	UPROPERTY(Transient, VisibleAnywhere, BlueprintReadOnly, Category=State)
 	float SelectedTime = 0.f;
 	
+	// SelectedAnim at SelectedTime is from the motion matching continuing pose search
 	UPROPERTY(Transient, VisibleAnywhere, BlueprintReadOnly, Category=State)
 	bool bIsContinuingPoseSearch = false;
 	
+	// SelectedAnim suggested play rate (it's usually calculated by estimating the ratio between the trajectory speed in the query against the speed in the result)
 	UPROPERTY(Transient, VisibleAnywhere, BlueprintReadOnly, Category=State)
 	float WantedPlayRate = 0.f;
 
+	// SelectedAnim associated looping state
 	UPROPERTY(Transient, VisibleAnywhere, BlueprintReadOnly, Category=State)
 	bool bLoop = false;
 	
+	// SelectedAnim associated mirror state
 	UPROPERTY(Transient, VisibleAnywhere, BlueprintReadOnly, Category=State)
 	bool bIsMirrored = false;
 	
+	// SelectedAnim associated BlendParameters (if SelectedAnim is a UBlendSpace)
 	UPROPERTY(Transient, VisibleAnywhere, BlueprintReadOnly, Category=State)
 	FVector BlendParameters = FVector::ZeroVector;
 
+	// database containig the SelectedAnim
 	UPROPERTY(Transient, VisibleAnywhere, BlueprintReadOnly, Category=State)
 	TWeakObjectPtr<const UPoseSearchDatabase> SelectedDatabase = nullptr;
 
+	// associated motion matching search cost for this result (how bad the SelectedAnim at SelectedTime etc looks in comparison to the query. The bigger the number, the more different the two are)
 	UPROPERTY(Transient, VisibleAnywhere, BlueprintReadOnly, Category=State)
 	float SearchCost = MAX_flt;
+
+	// bIsInteraction is true if this is a result from a motion matching interaction search
+	UPROPERTY(Transient, VisibleAnywhere, BlueprintReadOnly, Category="Experimental|Interaction", meta = (Experimental))
+	bool bIsInteraction = false;
+
+	// assigned role to this character. If this result is from a motion matched interaction, every interacting Actors will be assigned a different Role
+	UPROPERTY(Transient, VisibleAnywhere, BlueprintReadOnly, Category="Experimental|Interaction", meta = (Experimental))
+	FName Role;
+
+	// cached actors root transforms for all the roles in SelectedAnim (as UMultiAnimAsset). For all the optional missing roles (where the associated AnimContexts[i] is null), the ActorRootTransform will be identity
+	UPROPERTY(Transient, VisibleAnywhere, BlueprintReadOnly, Category="Experimental|Interaction", meta = (Experimental))
+	TArray<FTransform> ActorRootTransforms;
+
+	// cached actors root bone transforms for all the roles in SelectedAnim (as UMultiAnimAsset). For all the optional missing roles (where the associated AnimContexts[i] is null), the ActorRootBoneTransform will be identity
+	UPROPERTY(Transient, VisibleAnywhere, BlueprintReadOnly, Category="Experimental|Interaction", meta = (Experimental))
+	TArray<FTransform> ActorRootBoneTransforms;
+
+	// Experimental, this feature might be removed without warning, not for production use
+	// AnimContexts involved in this interaction for all the roles in SelectedAnim (as UMultiAnimAsset). If the role is missing because optional the AnimContext will be null
+	// BE CAREFUL HOW YOU USE THEM! since accessing data from other actors without proper tick dependencies leads to threading problems!
+	TArray<TWeakObjectPtr<const UObject>> AnimContexts;
 };
 
diff --git a/Engine/Plugins/Experimental/AnimNextPoseSearch/Source/AnimNextPoseSearch/Private/MotionMatchingTrait.cpp b/Engine/Plugins/Experimental/AnimNextPoseSearch/Source/AnimNextPoseSearch/Private/MotionMatchingTrait.cpp
--- a/Engine/Plugins/Experimental/AnimNextPoseSearch/Source/AnimNextPoseSearch/Private/MotionMatchingTrait.cpp
+++ b/Engine/Plugins/Experimental/AnimNextPoseSearch/Source/AnimNextPoseSearch/Private/MotionMatchingTrait.cpp
@@ -55,11 +55,11 @@ namespace UE::AnimNext
 		GeneratorMacro(IUpdate) \
 		GeneratorMacro(IEvaluate) \
 
-		// Trait implementation boilerplate
-		GENERATE_ANIM_TRAIT_IMPLEMENTATION(FMotionMatchingTrait, TRAIT_INTERFACE_ENUMERATOR, NULL_ANIM_TRAIT_INTERFACE_ENUMERATOR, NULL_ANIM_TRAIT_EVENT_ENUMERATOR)
+	// Trait implementation boilerplate
+	GENERATE_ANIM_TRAIT_IMPLEMENTATION(FMotionMatchingTrait, TRAIT_INTERFACE_ENUMERATOR, NULL_ANIM_TRAIT_INTERFACE_ENUMERATOR, NULL_ANIM_TRAIT_EVENT_ENUMERATOR)
 #undef TRAIT_INTERFACE_ENUMERATOR
 
-		void FMotionMatchingTrait::OnBecomeRelevant(FUpdateTraversalContext& Context, const TTraitBinding<IUpdate>& Binding, const FTraitUpdateState& TraitState) const
+	void FMotionMatchingTrait::OnBecomeRelevant(FUpdateTraversalContext& Context, const TTraitBinding<IUpdate>& Binding, const FTraitUpdateState& TraitState) const
 	{
 		IUpdate::OnBecomeRelevant(Context, Binding, TraitState);
 
@@ -71,8 +71,20 @@ namespace UE::AnimNext
 	{
 		using namespace UE::PoseSearch;
 
+		FInstanceData* InstanceData = Binding.GetInstanceData<FInstanceData>();
 		IUpdate::PreUpdate(Context, Binding, TraitState);
 
+#if WITH_EDITOR
+		if (InstanceData->bIsPostEvaluateBeingCalled)
+		{
+			InstanceData->bIsPostEvaluateBeingCalled = false;
+		}
+		else
+		{
+			UE_LOG(LogPoseSearch, Error, TEXT("FMotionMatchingTrait::PreUpdate, PostEvaluate has not being called last frame! Some trait in the TraitStack didn't propagate correctly the PostEvaluate!"));
+		}
+#endif // WITH_EDITOR
+
 		TTraitBinding<IBlendStack> BlendStackBinding;
 		if (!Binding.GetStackInterface(BlendStackBinding))
 		{
@@ -102,12 +114,8 @@ namespace UE::AnimNext
 		}
 
 		const FSharedData* SharedData = Binding.GetSharedData<FSharedData>();
-		check(SharedData);
-
-		FInstanceData* InstanceData = Binding.GetInstanceData<FInstanceData>();
-		check(InstanceData);
 		FMotionMatchingState& MotionMatchingState = InstanceData->MotionMatchingState;
-		FPoseSearchInteractionBlueprintResult& MotionMatchingInteractionState = InstanceData->MotionMatchingInteractionState;
+		FPoseSearchBlueprintResult& MotionMatchingInteractionState = InstanceData->MotionMatchingInteractionState;
 
 		bool bNeedsReset = false;
 #if WITH_EDITOR
@@ -126,7 +134,7 @@ namespace UE::AnimNext
 		if (bNeedsReset)
 		{
 			MotionMatchingState.Reset();
-			MotionMatchingInteractionState = FPoseSearchInteractionBlueprintResult();
+			MotionMatchingInteractionState = FPoseSearchBlueprintResult();
 
 			// @todo: implement me!
 			//BlendStackBinding.ResetBlendStack();
@@ -177,7 +185,7 @@ namespace UE::AnimNext
 		const TArray<FPoseSearchInteractionAvailability>& Availabilities = SharedData->GetAvailabilities(Binding);
 		if (Availabilities.IsEmpty())
 		{
-			MotionMatchingInteractionState = FPoseSearchInteractionBlueprintResult();
+			MotionMatchingInteractionState = FPoseSearchBlueprintResult();
 		}
 		else
 		{
@@ -186,22 +194,22 @@ namespace UE::AnimNext
 
 			// we're NOT providing the UMultiAnimAsset here, relying on the interaction subsystem to resolve the ContinuingProperties properly
 			FPoseSearchContinuingProperties ContinuingProperties;
-			ContinuingProperties.PlayingAsset = MotionMatchingInteractionState.SelectedAnimation; // GetAnimAsset();
+			ContinuingProperties.PlayingAsset = MotionMatchingInteractionState.SelectedAnim.Get(); // GetAnimAsset();
 			ContinuingProperties.PlayingAssetAccumulatedTime = MotionMatchingInteractionState.SelectedTime;
 
 			UPoseSearchInteractionLibrary::MotionMatchInteraction(MotionMatchingInteractionState, Availabilities, AnimContext, ContinuingProperties, FName(), PoseHistory, SharedData->GetbValidateResultAgainstAvailabilities(Binding));
 			check(MotionMatchingInteractionState.ActorRootTransforms.Num() == MotionMatchingInteractionState.ActorRootBoneTransforms.Num());
 		}
 
-		if (MotionMatchingInteractionState.SelectedAnimation)
+		if (MotionMatchingInteractionState.SelectedAnim != nullptr)
 		{
 			if (!MotionMatchingInteractionState.bIsContinuingPoseSearch)
 			{
-				if (const UMultiAnimAsset* MultiAnimAsset = Cast<UMultiAnimAsset>(MotionMatchingInteractionState.SelectedAnimation))
+				if (const UMultiAnimAsset* MultiAnimAsset = Cast<UMultiAnimAsset>(MotionMatchingInteractionState.SelectedAnim))
 				{
 					AnimationAsset = MultiAnimAsset->GetAnimationAsset(MotionMatchingInteractionState.Role);
 				}
-				else if (UAnimationAsset* SingleAnimAsset = Cast<UAnimationAsset>(MotionMatchingInteractionState.SelectedAnimation))
+				else if (const UAnimationAsset* SingleAnimAsset = Cast<UAnimationAsset>(MotionMatchingInteractionState.SelectedAnim))
 				{
 					AnimationAsset = SingleAnimAsset;
 				}
@@ -218,7 +226,11 @@ namespace UE::AnimNext
 			const TConstArrayView<TObjectPtr<const UPoseSearchDatabase>> DatabasesToSearch = SharedData->GetDatabases(Binding);
 			if (DatabasesToSearch.IsEmpty())
 			{
-				UE_LOG(LogPoseSearch, Error, TEXT("FMotionMatchingTrait::PreUpdate, No database assets provided for motion matching."));
+				// if we have availabilities, it's ok to have a FMotionMatchingTrait set up purely for interactions, without any other database to search! 
+				if (Availabilities.IsEmpty())
+				{
+					UE_LOG(LogPoseSearch, Error, TEXT("FMotionMatchingTrait::PreUpdate, No database assets provided for motion matching."));
+				}
 			}
 			else
 			{
@@ -250,18 +262,24 @@ namespace UE::AnimNext
 				WantedPlayRate = MotionMatchingState.WantedPlayRate;
 			}
 
-			const bool bJumpToPose = MotionMatchingState.CurrentSearchResult.IsValid() && !MotionMatchingState.CurrentSearchResult.bIsContinuingPoseSearch;
-			if (bJumpToPose)
+			const UPoseSearchDatabase* Database = MotionMatchingState.CurrentSearchResult.Database.Get();
+#if WITH_EDITOR
+			// we need to validate if Database is not reindexing because of any possible throttling of the search!
+			if (EAsyncBuildIndexResult::Success == FAsyncPoseSearchDatabasesManagement::RequestAsyncBuildIndex(Database, ERequestAsyncBuildFlag::ContinueRequest))
+#endif // WITH_EDITOR
 			{
-				if (const UE::PoseSearch::FSearchIndexAsset* SearchIndexAsset = MotionMatchingState.CurrentSearchResult.GetSearchIndexAsset())
+				const bool bJumpToPose = MotionMatchingState.CurrentSearchResult.IsValid() && !MotionMatchingState.CurrentSearchResult.bIsContinuingPoseSearch;
+				if (bJumpToPose)
 				{
-					const UPoseSearchDatabase* Database = MotionMatchingState.CurrentSearchResult.Database.Get();
-					if (const FPoseSearchDatabaseAnimationAssetBase* DatabaseAnimationAssetBase = Database->GetDatabaseAnimationAsset<FPoseSearchDatabaseAnimationAssetBase>(*SearchIndexAsset))
+					if (const UE::PoseSearch::FSearchIndexAsset* SearchIndexAsset = MotionMatchingState.CurrentSearchResult.GetSearchIndexAsset())
 					{
-						AnimationAsset = DatabaseAnimationAssetBase->GetAnimationAssetForRole(Database->Schema->GetDefaultRole());
-						AssetTime = MotionMatchingState.CurrentSearchResult.AssetTime;
-						bLoop = SearchIndexAsset->IsLooping();
-						BlendSpaceParameters = SearchIndexAsset->GetBlendParameters();
+						if (const FPoseSearchDatabaseAnimationAssetBase* DatabaseAnimationAssetBase = Database->GetDatabaseAnimationAsset<FPoseSearchDatabaseAnimationAssetBase>(*SearchIndexAsset))
+						{
+							AnimationAsset = DatabaseAnimationAssetBase->GetAnimationAssetForRole(Database->Schema->GetDefaultRole());
+							AssetTime = MotionMatchingState.CurrentSearchResult.AssetTime;
+							bLoop = SearchIndexAsset->IsLooping();
+							BlendSpaceParameters = SearchIndexAsset->GetBlendParameters();
+						}
 					}
 				}
 			}
@@ -342,46 +360,86 @@ namespace UE::AnimNext
 			}
 		}
 
+		PublishResults(Binding);
+	}
+
+	void FMotionMatchingTrait::PublishResults(const TTraitBinding<IUpdate>& Binding) const
+	{
+		using namespace UE::PoseSearch;
+
+		const FSharedData* SharedData = Binding.GetSharedData<FSharedData>();
+		const FInstanceData* InstanceData = Binding.GetInstanceData<FInstanceData>();
+		const FMotionMatchingState& MotionMatchingState = InstanceData->MotionMatchingState;
+		const FPoseSearchBlueprintResult& MotionMatchingInteractionState = InstanceData->MotionMatchingInteractionState;
+
 		FName ResultOutputVariableName = SharedData->GetMotionMatchingResultVariableName(Binding);
-		const FAnimNextModuleHandle ResultModuleHandle = SharedData->GetMotionMatchingResultModuleHandle(Binding);
-		if (ResultOutputVariableName != NAME_None && ResultModuleHandle.ModuleHandle.IsValid())
+		if (ResultOutputVariableName == NAME_None)
 		{
-			FAnimNextModuleInstance* ModuleInstance = Binding.GetTraitPtr().GetNodeInstance()->GetOwner().GetModuleInstance();
+			return;
+		}
 
-			const UPoseSearchDatabase* Database = MotionMatchingState.CurrentSearchResult.Database.Get();
+		FPoseSearchBlueprintResult Result;
 
-			if (const FSearchIndexAsset* SearchIndexAsset = MotionMatchingState.CurrentSearchResult.GetSearchIndexAsset())
+		const bool bIsMotionMatchingInteractionStateValid = MotionMatchingInteractionState.SelectedAnim != nullptr;
+		const UPoseSearchDatabase* CurrentResultDatabase = bIsMotionMatchingInteractionStateValid ? MotionMatchingInteractionState.SelectedDatabase.Get() : MotionMatchingState.CurrentSearchResult.Database.Get();
+		if (CurrentResultDatabase == nullptr || CurrentResultDatabase->Schema == nullptr)
+		{
+			// @todo: should we log a warning?
+		}
+
+#if WITH_EDITOR
+		else if (EAsyncBuildIndexResult::Success != FAsyncPoseSearchDatabasesManagement::RequestAsyncBuildIndex(CurrentResultDatabase, ERequestAsyncBuildFlag::ContinueRequest))
+		{
+			UE_LOG(LogPoseSearch, Warning, TEXT("FMotionMatchingTrait::PreUpdate called while '%s' is indexing. returning an invalid result.."), *CurrentResultDatabase->GetName());
+		}
+		// in case this node is not updated, and MotionMatchingState.CurrentSearchResult.Database gets modified, we could end up with CurrentSearchResult being out of synch with the updated database, so we need to reset the state
+		else if (!bIsMotionMatchingInteractionStateValid && MotionMatchingState.CurrentSearchResult.IsValid() && MotionMatchingState.CurrentSearchResult.PoseIdx >= MotionMatchingState.CurrentSearchResult.Database->GetSearchIndex().GetNumPoses())
+		{
+			UE_LOG(LogPoseSearch, Warning, TEXT("FMotionMatchingTrait::PreUpdate inconsistency between MotionMatchingState and '%s', probably because of stale MotionMatchingState after database got reindexed. returning an invalid result.."), *CurrentResultDatabase->GetName());
+		}
+#endif // WITH_EDITOR
+		else if (bIsMotionMatchingInteractionStateValid)
+		{
+			Result = MotionMatchingInteractionState;
+		}
+		else if (const FSearchIndexAsset* SearchIndexAsset = MotionMatchingState.CurrentSearchResult.GetSearchIndexAsset())
+		{
+			if (const FPoseSearchDatabaseAnimationAssetBase* DatabaseAsset = CurrentResultDatabase->GetDatabaseAnimationAsset<FPoseSearchDatabaseAnimationAssetBase>(*SearchIndexAsset))
 			{
-				const UPoseSearchDatabase* CurrentResultDatabase = MotionMatchingState.CurrentSearchResult.Database.Get();
-				if (CurrentResultDatabase && CurrentResultDatabase->Schema)
+				Result.SelectedAnim = DatabaseAsset->GetAnimationAsset();
+				Result.SelectedTime = MotionMatchingState.CurrentSearchResult.AssetTime;
+				Result.bIsContinuingPoseSearch = MotionMatchingState.CurrentSearchResult.bIsContinuingPoseSearch;
+				Result.bLoop = SearchIndexAsset->IsLooping();
+				Result.bIsMirrored = SearchIndexAsset->IsMirrored();
+				Result.BlendParameters = SearchIndexAsset->GetBlendParameters();
+				Result.SelectedDatabase = CurrentResultDatabase;
+				Result.WantedPlayRate = MotionMatchingState.WantedPlayRate;
+				Result.SearchCost = MotionMatchingState.CurrentSearchResult.PoseCost;
+			}
+		}
+
+		FAnimNextModuleInstance* ModuleInstance = Binding.GetTraitPtr().GetNodeInstance()->GetOwner().GetModuleInstance();
+		check(ModuleInstance);
+
+		const FAnimNextModuleHandle ResultModuleHandle = SharedData->GetMotionMatchingResultModuleHandle(Binding);
+		if (ResultModuleHandle.ModuleHandle.IsValid())
+		{
+			ModuleInstance->QueueTaskOnOtherModule(
+				ResultModuleHandle.ModuleHandle,
+				FRigUnit_AnimNextPrePhysicsEvent::EventName,
+				[Result, ResultOutputVariableName](const UE::AnimNext::FModuleTaskContext& TaskContext)
 				{
-					if (const FPoseSearchDatabaseAnimationAssetBase* DatabaseAsset = CurrentResultDatabase->GetDatabaseAnimationAsset<FPoseSearchDatabaseAnimationAssetBase>(*SearchIndexAsset))
+					FAnimNextModuleInstance* OtherModuleInstance = TaskContext.GetModuleInstance();
+					if (OtherModuleInstance)
 					{
-						FPoseSearchBlueprintResult Result;
-						Result.SelectedAnimation = DatabaseAsset->GetAnimationAsset();
-						Result.SelectedTime = MotionMatchingState.CurrentSearchResult.AssetTime;
-						Result.bIsContinuingPoseSearch = MotionMatchingState.CurrentSearchResult.bIsContinuingPoseSearch;
-						Result.bLoop = SearchIndexAsset->IsLooping();
-						Result.bIsMirrored = SearchIndexAsset->IsMirrored();
-						Result.BlendParameters = SearchIndexAsset->GetBlendParameters();
-						Result.SelectedDatabase = CurrentResultDatabase;
-						Result.WantedPlayRate = MotionMatchingState.WantedPlayRate;
-
-						ModuleInstance->QueueTaskOnOtherModule(
-							ResultModuleHandle.ModuleHandle,
-							FRigUnit_AnimNextPrePhysicsEvent::EventName,
-							[Result, ResultOutputVariableName](const UE::AnimNext::FModuleTaskContext& TaskContext)
-								{
-									FAnimNextModuleInstance* OtherModuleInstance = TaskContext.GetModuleInstance();
-									if (OtherModuleInstance)
-									{
-										OtherModuleInstance->SetVariable(ResultOutputVariableName, Result);
-									}
-								},
-							UE::AnimNext::ETaskRunLocation::Before);
+						OtherModuleInstance->SetVariable(ResultOutputVariableName, Result);
 					}
-				}
-			}
+				},
+				UE::AnimNext::ETaskRunLocation::Before);
+		}
+		else
+		{
+			ModuleInstance->SetVariable(ResultOutputVariableName, Result);
 		}
 	}
 
@@ -392,16 +450,26 @@ namespace UE::AnimNext
 		IEvaluate::PostEvaluate(Context, Binding);
 
 		FInstanceData* InstanceData = Binding.GetInstanceData<FInstanceData>();
-		check(InstanceData);
+
+#if WITH_EDITOR
+		if (InstanceData->bIsPostEvaluateBeingCalled)
+		{
+			UE_LOG(LogPoseSearch, Error, TEXT("FMotionMatchingTrait::PostEvaluate, PostEvaluate called without calling PreUpdate on this frame! Some trait in the TraitStack doesn't propagate correctly the PreUpdate!"));
+		}
+		else
+		{
+			InstanceData->bIsPostEvaluateBeingCalled = true;
+		}
+#endif // WITH_EDITOR
 
 		// if it's not a multi character interaction we can skip the warping logic entirely
-		if (!InstanceData->MotionMatchingInteractionState.SelectedAnimation)
+		if (InstanceData->MotionMatchingInteractionState.SelectedAnim == nullptr)
 		{
 			return;
 		}
 
-		const FPoseSearchInteractionBlueprintResult& CurrentResult = InstanceData->MotionMatchingInteractionState;
-		if (const UMultiAnimAsset* MultiAnimAsset = Cast<UMultiAnimAsset>(CurrentResult.SelectedAnimation))
+		const FPoseSearchBlueprintResult& CurrentResult = InstanceData->MotionMatchingInteractionState;
+		if (const UMultiAnimAsset* MultiAnimAsset = Cast<UMultiAnimAsset>(CurrentResult.SelectedAnim))
 		{
 			const int32 CurrentResultRoleIndex = GetRoleIndex(MultiAnimAsset, CurrentResult.Role);
 			if (CurrentResultRoleIndex != INDEX_NONE)
@@ -459,8 +527,8 @@ void FAnimNextMotionMatchingTask::Execute(UE::AnimNext::FEvaluationVM& VM) const
 		return;
 	}
 
-	const FPoseSearchInteractionBlueprintResult& CurrentResult = InstanceData->MotionMatchingInteractionState;
-	const UMultiAnimAsset* MultiAnimAsset = Cast<UMultiAnimAsset>(CurrentResult.SelectedAnimation);
+	const FPoseSearchBlueprintResult& CurrentResult = InstanceData->MotionMatchingInteractionState;
+	const UMultiAnimAsset* MultiAnimAsset = Cast<UMultiAnimAsset>(CurrentResult.SelectedAnim);
 	if (!MultiAnimAsset)
 	{
 		return;
diff --git a/Engine/Plugins/Experimental/AnimNextPoseSearch/Source/AnimNextPoseSearch/Private/MotionMatchingTrait.h b/Engine/Plugins/Experimental/AnimNextPoseSearch/Source/AnimNextPoseSearch/Private/MotionMatchingTrait.h
--- a/Engine/Plugins/Experimental/AnimNextPoseSearch/Source/AnimNextPoseSearch/Private/MotionMatchingTrait.h
+++ b/Engine/Plugins/Experimental/AnimNextPoseSearch/Source/AnimNextPoseSearch/Private/MotionMatchingTrait.h
@@ -70,7 +70,7 @@ struct FMotionMatchingTraitSharedData : public FAnimNextTraitSharedData
 	UPROPERTY(EditAnywhere, Category = Settings)
 	FName MotionMatchingResultVariableName = NAME_None;
 
-	// If set along with result variable name, the motion matching result will be written to the variable in this module
+	// If set along with result variable name, the motion matching result will be written to the variable in this module, otherwise the variable will be written in the module containing the MotionMatchingTrait
 	UPROPERTY(EditAnywhere, Category = Settings)
 	FAnimNextModuleHandle MotionMatchingResultModuleHandle;
 
@@ -129,7 +129,11 @@ struct FMotionMatchingTrait : FAdditiveTrait, IUpdate, IEvaluate
 		FMotionMatchingState MotionMatchingState;
 
 		// @todo: merge this with MotionMatchingState
-		FPoseSearchInteractionBlueprintResult MotionMatchingInteractionState;
+		FPoseSearchBlueprintResult MotionMatchingInteractionState;
+
+#if WITH_EDITOR
+		bool bIsPostEvaluateBeingCalled = true;
+#endif // WITH_EDITOR
 	};
 
 	// IUpdate impl
@@ -138,6 +142,9 @@ struct FMotionMatchingTrait : FAdditiveTrait, IUpdate, IEvaluate
 
 	// IEvaluate impl
 	virtual void PostEvaluate(FEvaluateTraversalContext& Context, const TTraitBinding<IEvaluate>& Binding) const override;
+
+private:
+	void PublishResults(const TTraitBinding<IUpdate>& Binding) const;
 };
 
 } // namespace UE::AnimNext
