Anim Next: Motion Matching skip blend if we're blending into the same blendspace

diff --git a/Engine/Plugins/Experimental/AnimNextPoseSearch/Source/AnimNextPoseSearch/Private/MotionMatchingTrait.cpp b/Engine/Plugins/Experimental/AnimNextPoseSearch/Source/AnimNextPoseSearch/Private/MotionMatchingTrait.cpp
--- a/Engine/Plugins/Experimental/AnimNextPoseSearch/Source/AnimNextPoseSearch/Private/MotionMatchingTrait.cpp
+++ b/Engine/Plugins/Experimental/AnimNextPoseSearch/Source/AnimNextPoseSearch/Private/MotionMatchingTrait.cpp
@@ -290,45 +290,62 @@ namespace UE::AnimNext
 		DesiredPlayRate = GVarAnimNextMotionMatchingPlayRateEnabled == UE::Private::EPlayRateState::Enabled ? DesiredPlayRate : 1.0f;
 #endif // ENABLE_ANIM_DEBUG
 
+		const bool bUpdateBlendSpaceInputs = SharedData->GetbUpdateBlendSpaceInputs(Binding);
 		if (AnimationAsset)
 		{
-			UE::AnimNext::IBlendStack::FGraphRequest NewGraphRequest;
-			NewGraphRequest.BlendArgs = SharedData->BlendArguments;
+			const bool bIsBlendSpace = AnimationAsset->IsA<UBlendSpace>();
+			bool bSkipBlend = false;
+			if (bIsBlendSpace && bUpdateBlendSpaceInputs)
+			{
+				IBlendStack::FGraphRequestPtr ActiveGraphRequest;
+				BlendStackBinding.GetActiveGraph(Context, ActiveGraphRequest);
+				if (ActiveGraphRequest && (ActiveGraphRequest->FactoryObject == AnimationAsset))
+				{
+					// If we're still on the same blend space, and we're updating blend space inputs dynamically, then DON'T BLEND.
+					bSkipBlend = true;
+				}
+			}
 
-			// Create a default 'native data interface' payload via which we expect to communicate with the graph 
-			FInstancedStruct NewPayload;
-			// @TODO: Support asset->playload mapping in a more data-driven way
-			if (AnimationAsset->IsA<UBlendSpace>())
+			if (bSkipBlend == false)
 			{
-				FAnimNextNativeDataInterface_BlendSpacePlayer& NewNativeInterface = NewPayload.InitializeAs<FAnimNextNativeDataInterface_BlendSpacePlayer>();
-				NewNativeInterface.BlendSpace = Cast<UBlendSpace>(AnimationAsset);
+				UE::AnimNext::IBlendStack::FGraphRequest NewGraphRequest;
+				NewGraphRequest.BlendArgs = SharedData->BlendArguments;
 
-				NewNativeInterface.XAxisSamplePoint = BlendSpaceParameters.X;
-				NewNativeInterface.YAxisSamplePoint = BlendSpaceParameters.Y;
+				// Create a default 'native data interface' payload via which we expect to communicate with the graph 
+				FInstancedStruct NewPayload;
+				// @TODO: Support asset->playload mapping in a more data-driven way
+				if (bIsBlendSpace)
+				{
+					FAnimNextNativeDataInterface_BlendSpacePlayer& NewNativeInterface = NewPayload.InitializeAs<FAnimNextNativeDataInterface_BlendSpacePlayer>();
+					NewNativeInterface.BlendSpace = Cast<UBlendSpace>(AnimationAsset);
 
-				NewNativeInterface.PlayRate = DesiredPlayRate;
-				NewNativeInterface.StartPosition = MotionMatchingState.CurrentSearchResult.AssetTime;
-				NewNativeInterface.Loop = bLoop;
-			}
-			else
-			{
-				FAnimNextNativeDataInterface_AnimSequencePlayer& NewNativeInterface = NewPayload.InitializeAs<FAnimNextNativeDataInterface_AnimSequencePlayer>();
-				NewNativeInterface.AnimSequence = Cast<UAnimSequence>(AnimationAsset);
-				NewNativeInterface.PlayRate = DesiredPlayRate;
-				NewNativeInterface.StartPosition = MotionMatchingState.CurrentSearchResult.AssetTime;
-				NewNativeInterface.Loop = bLoop;
-			}
+					NewNativeInterface.XAxisSamplePoint = BlendSpaceParameters.X;
+					NewNativeInterface.YAxisSamplePoint = BlendSpaceParameters.Y;
 
-			FAnimNextDataInterfacePayload GraphPayload;
-			GraphPayload.AddNative(MoveTemp(NewPayload));
+					NewNativeInterface.PlayRate = DesiredPlayRate;
+					NewNativeInterface.StartPosition = MotionMatchingState.CurrentSearchResult.AssetTime;
+					NewNativeInterface.Loop = bLoop;
+				}
+				else
+				{
+					FAnimNextNativeDataInterface_AnimSequencePlayer& NewNativeInterface = NewPayload.InitializeAs<FAnimNextNativeDataInterface_AnimSequencePlayer>();
+					NewNativeInterface.AnimSequence = Cast<UAnimSequence>(AnimationAsset);
+					NewNativeInterface.PlayRate = DesiredPlayRate;
+					NewNativeInterface.StartPosition = MotionMatchingState.CurrentSearchResult.AssetTime;
+					NewNativeInterface.Loop = bLoop;
+				}
 
-			const UAnimNextAnimationGraph* AnimationGraph = IGraphFactory::GetGraphFromObjectWithFallback(Context, Binding, AnimationAsset, GraphPayload);
-			if (AnimationGraph != nullptr)
-			{
-				NewGraphRequest.GraphPayload = MoveTemp(GraphPayload);
-				NewGraphRequest.AnimationGraph = AnimationGraph;
+				FAnimNextDataInterfacePayload GraphPayload;
+				GraphPayload.AddNative(MoveTemp(NewPayload));
 
-				BlendStackBinding.PushGraph(Context, MoveTemp(NewGraphRequest));
+				const UAnimNextAnimationGraph* AnimationGraph = IGraphFactory::GetGraphFromObjectWithFallback(Context, Binding, AnimationAsset, GraphPayload);
+				if (AnimationGraph != nullptr)
+				{
+					NewGraphRequest.GraphPayload = MoveTemp(GraphPayload);
+					NewGraphRequest.AnimationGraph = AnimationGraph;
+					NewGraphRequest.FactoryObject = AnimationAsset;
+					BlendStackBinding.PushGraph(Context, MoveTemp(NewGraphRequest));
+				}
 			}
 		}
 
@@ -350,7 +367,7 @@ namespace UE::AnimNext
 				else if (FAnimNextNativeDataInterface_BlendSpacePlayer* BlendSpacePayload = Payload.GetPtr<FAnimNextNativeDataInterface_BlendSpacePlayer>())
 				{
 					BlendSpacePayload->PlayRate = DesiredPlayRate;
-					if (SharedData->GetbUpdateBlendSpaceInputs(Binding))
+					if (bUpdateBlendSpaceInputs)
 					{
 						BlendSpacePayload->XAxisSamplePoint = SharedData->GetXAxisSamplePoint(Binding);
 						BlendSpacePayload->YAxisSamplePoint = SharedData->GetYAxisSamplePoint(Binding);
