PoseSearch - AnimNext MM interaction blendstack payload AssetTime fix

diff --git a/Engine/Plugins/Experimental/AnimNextPoseSearch/Source/AnimNextPoseSearch/Private/MotionMatchingTrait.cpp b/Engine/Plugins/Experimental/AnimNextPoseSearch/Source/AnimNextPoseSearch/Private/MotionMatchingTrait.cpp
--- a/Engine/Plugins/Experimental/AnimNextPoseSearch/Source/AnimNextPoseSearch/Private/MotionMatchingTrait.cpp
+++ b/Engine/Plugins/Experimental/AnimNextPoseSearch/Source/AnimNextPoseSearch/Private/MotionMatchingTrait.cpp
@@ -323,7 +323,7 @@ namespace UE::AnimNext
 					NewNativeInterface.YAxisSamplePoint = BlendSpaceParameters.Y;
 
 					NewNativeInterface.PlayRate = DesiredPlayRate;
-					NewNativeInterface.StartPosition = MotionMatchingState.CurrentSearchResult.AssetTime;
+					NewNativeInterface.StartPosition = AssetTime;
 					NewNativeInterface.Loop = bLoop;
 				}
 				else
@@ -331,7 +331,7 @@ namespace UE::AnimNext
 					FAnimNextNativeDataInterface_AnimSequencePlayer& NewNativeInterface = NewPayload.InitializeAs<FAnimNextNativeDataInterface_AnimSequencePlayer>();
 					NewNativeInterface.AnimSequence = Cast<UAnimSequence>(AnimationAsset);
 					NewNativeInterface.PlayRate = DesiredPlayRate;
-					NewNativeInterface.StartPosition = MotionMatchingState.CurrentSearchResult.AssetTime;
+					NewNativeInterface.StartPosition = AssetTime;
 					NewNativeInterface.Loop = bLoop;
 				}
 
