AnimNextPoseSearch - logging errors in case FHistoryCollectorTrait::PostEvaluate is not called properly

diff --git a/Engine/Plugins/Experimental/AnimNextPoseSearch/Source/AnimNextPoseSearch/Private/HistoryCollectorTrait.cpp b/Engine/Plugins/Experimental/AnimNextPoseSearch/Source/AnimNextPoseSearch/Private/HistoryCollectorTrait.cpp
--- a/Engine/Plugins/Experimental/AnimNextPoseSearch/Source/AnimNextPoseSearch/Private/HistoryCollectorTrait.cpp
+++ b/Engine/Plugins/Experimental/AnimNextPoseSearch/Source/AnimNextPoseSearch/Private/HistoryCollectorTrait.cpp
@@ -159,6 +159,17 @@ namespace UE::AnimNext
 		const FSharedData* SharedData = Binding.GetSharedData<FSharedData>();
 		FInstanceData* InstanceData = Binding.GetInstanceData<FInstanceData>();
 		
+#if WITH_EDITOR
+		if (InstanceData->bIsPostEvaluateBeingCalled)
+		{
+			InstanceData->bIsPostEvaluateBeingCalled = false;
+		}
+		else
+		{
+			UE_LOG(LogPoseSearch, Error, TEXT("FHistoryCollectorTrait::PreUpdate, PostEvaluate has not being called last frame! Some trait in the TraitStack didn't propagate correctly the PostEvaluate!"));
+		}
+#endif // WITH_EDITOR
+
 		InstanceData->DeltaTime = TraitState.GetDeltaTime();
 
 		// @todo: implement bResetOnBecomingRelevant
@@ -200,6 +211,17 @@ namespace UE::AnimNext
 		FInstanceData* InstanceData = Binding.GetInstanceData<FInstanceData>();
 		const bool bStoreScales = SharedData->GetbStoreScales(Binding);
 
+#if WITH_EDITOR
+		if (InstanceData->bIsPostEvaluateBeingCalled)
+		{
+			UE_LOG(LogPoseSearch, Error, TEXT("FHistoryCollectorTrait::PostEvaluate, PostEvaluate called without calling PreUpdate on this frame! Some trait in the TraitStack doesn't propagate correctly the PreUpdate!"));
+		}
+		else
+		{
+			InstanceData->bIsPostEvaluateBeingCalled = true;
+		}
+#endif // WITH_EDITOR
+
 		check(InstanceData->PoseHistoryPtr.IsValid());
 		FAnimNextHistoryCollectorTask Task;
 		Task.PoseHistory = InstanceData->PoseHistoryPtr.Get();
diff --git a/Engine/Plugins/Experimental/AnimNextPoseSearch/Source/AnimNextPoseSearch/Private/HistoryCollectorTrait.h b/Engine/Plugins/Experimental/AnimNextPoseSearch/Source/AnimNextPoseSearch/Private/HistoryCollectorTrait.h
--- a/Engine/Plugins/Experimental/AnimNextPoseSearch/Source/AnimNextPoseSearch/Private/HistoryCollectorTrait.h
+++ b/Engine/Plugins/Experimental/AnimNextPoseSearch/Source/AnimNextPoseSearch/Private/HistoryCollectorTrait.h
@@ -138,6 +138,10 @@ namespace UE::AnimNext
 			FInstanceData() : PoseHistoryPtr(MakeShareable(new UE::PoseSearch::FGenerateTrajectoryPoseHistory())) { }
 			TSharedPtr<UE::PoseSearch::FGenerateTrajectoryPoseHistory, ESPMode::ThreadSafe> PoseHistoryPtr;
 			float DeltaTime = 0.f;
+
+#if WITH_EDITOR
+			bool bIsPostEvaluateBeingCalled = true;
+#endif // WITH_EDITOR
 		};
 
 		// IUpdate impl
