PoseSearch - implementing missing MM node logic to MM trait

diff --git a/Engine/Plugins/Animation/PoseSearch/Source/Runtime/Private/AnimNode_MotionMatchingInteraction.cpp b/Engine/Plugins/Animation/PoseSearch/Source/Runtime/Private/AnimNode_MotionMatchingInteraction.cpp
--- a/Engine/Plugins/Animation/PoseSearch/Source/Runtime/Private/AnimNode_MotionMatchingInteraction.cpp
+++ b/Engine/Plugins/Animation/PoseSearch/Source/Runtime/Private/AnimNode_MotionMatchingInteraction.cpp
@@ -71,25 +71,21 @@ void FAnimNode_MotionMatchingInteraction::UpdateAssetPlayer(const FAnimationUpda
 		UPoseSearchInteractionLibrary::MotionMatchInteraction(CurrentResult, Availabilities, Context.AnimInstanceProxy->GetAnimInstanceObject(), ContinuingProperties, FName(),  & PoseHistoryProvider->GetPoseHistory(), bValidateResultAgainstAvailabilities);
 		check(CurrentResult.ActorRootTransforms.Num() == CurrentResult.ActorRootBoneTransforms.Num());
 
-		const UAnimationAsset* RoledAnimAsset = nullptr;
 		if (const UMultiAnimAsset* MultiAnimAsset = Cast<UMultiAnimAsset>(CurrentResult.SelectedAnim))
 		{
-			RoledAnimAsset = MultiAnimAsset->GetAnimationAsset(CurrentResult.Role);
-		}
-		else if (const UAnimationAsset* SingleAnimAsset = Cast<UAnimationAsset>(CurrentResult.SelectedAnim))
-		{
-			RoledAnimAsset = SingleAnimAsset;
-		}
-
-		if (RoledAnimAsset && (AnimPlayers.IsEmpty() || !CurrentResult.bIsContinuingPoseSearch))
-		{
-			const FPoseSearchRoledSkeleton* RoledSkeleton = CurrentResult.SelectedDatabase->Schema->GetRoledSkeleton(CurrentResult.Role);
-			check(RoledSkeleton);
-
-			BlendTo(Context, const_cast<UAnimationAsset*>(RoledAnimAsset), CurrentResult.SelectedTime, CurrentResult.bLoop, CurrentResult.bIsMirrored, RoledSkeleton->MirrorDataTable.Get(),
-				BlendTime, BlendProfile, BlendOption, bUseInertialBlend, NAME_None, CurrentResult.BlendParameters, CurrentResult.WantedPlayRate);
-
-			bBlendToExecuted = true;
+			if (const UAnimationAsset* RoledAnimAsset = MultiAnimAsset->GetAnimationAsset(CurrentResult.Role))
+			{
+				if (AnimPlayers.IsEmpty() || !CurrentResult.bIsContinuingPoseSearch)
+				{
+					const FPoseSearchRoledSkeleton* RoledSkeleton = CurrentResult.SelectedDatabase->Schema->GetRoledSkeleton(CurrentResult.Role);
+					check(RoledSkeleton);
+
+					BlendTo(Context, const_cast<UAnimationAsset*>(RoledAnimAsset), CurrentResult.SelectedTime, CurrentResult.bLoop, CurrentResult.bIsMirrored, RoledSkeleton->MirrorDataTable.Get(),
+						BlendTime, BlendProfile, BlendOption, bUseInertialBlend, NAME_None, CurrentResult.BlendParameters, CurrentResult.WantedPlayRate);
+
+					bBlendToExecuted = true;
+				}
+			}
 		}
 	}
 	else
diff --git a/Engine/Plugins/Experimental/AnimNextPoseSearch/Source/AnimNextPoseSearch/Private/MotionMatchingTrait.cpp b/Engine/Plugins/Experimental/AnimNextPoseSearch/Source/AnimNextPoseSearch/Private/MotionMatchingTrait.cpp
--- a/Engine/Plugins/Experimental/AnimNextPoseSearch/Source/AnimNextPoseSearch/Private/MotionMatchingTrait.cpp
+++ b/Engine/Plugins/Experimental/AnimNextPoseSearch/Source/AnimNextPoseSearch/Private/MotionMatchingTrait.cpp
@@ -114,46 +114,43 @@ namespace UE::AnimNext
 		}
 
 		const FSharedData* SharedData = Binding.GetSharedData<FSharedData>();
-		FMotionMatchingState& MotionMatchingState = InstanceData->MotionMatchingState;
-		FPoseSearchBlueprintResult& MotionMatchingInteractionState = InstanceData->MotionMatchingInteractionState;
+		const UObject* AnimationAsset = nullptr;
+		float AssetTime = 0.f;
+		float WantedPlayRate = 1.f;
+		bool bLoop = false;
+		FVector BlendSpaceParameters = FVector::ZeroVector;
 
-		bool bNeedsReset = false;
-#if WITH_EDITOR
-		if (EAsyncBuildIndexResult::Success != FAsyncPoseSearchDatabasesManagement::RequestAsyncBuildIndex(MotionMatchingState.CurrentSearchResult.Database.Get(), ERequestAsyncBuildFlag::ContinueRequest))
-		{
-			bNeedsReset = true;
-		}
-		// in case this node is not updated, and MotionMatchingState.CurrentSearchResult.Database gets modified, we could end up with CurrentSearchResult being out of synch with the updated database, so we need to reset the state
-		else if (MotionMatchingState.CurrentSearchResult.IsValid() && MotionMatchingState.CurrentSearchResult.PoseIdx >= MotionMatchingState.CurrentSearchResult.Database->GetSearchIndex().GetNumPoses())
+		const UObject* AnimContext = nullptr;
+		if (const FAnimNextModuleInstance* ModuleInstance = Context.GetRootGraphInstance().GetModuleInstance())
 		{
-			bNeedsReset = true;
+			AnimContext = ModuleInstance->GetObject();
 		}
-#endif // WITH_EDITOR
 
-		// If we just became relevant and haven't been initialized yet, then reset motion matching state, otherwise update the asset time using the player node.
-		if (bNeedsReset)
+		FMotionMatchingState& MotionMatchingState = InstanceData->MotionMatchingState;
+		FPoseSearchBlueprintResult& MotionMatchingInteractionState = InstanceData->MotionMatchingInteractionState;
+
+		const TArray<FPoseSearchInteractionAvailability>& Availabilities = SharedData->GetAvailabilities(Binding);
+		if (Availabilities.IsEmpty())
 		{
-			MotionMatchingState.Reset();
 			MotionMatchingInteractionState = FPoseSearchBlueprintResult();
-
-			// @todo: implement me!
-			//BlendStackBinding.ResetBlendStack();
 		}
 		else
 		{
-			// We adjust the motion matching state asset time to the current player node's asset time. This is done 
+			// WIP motion matched interactions
+			CheckInteractionThreadSafety(AnimContext);
+
+			// We adjust the MotionMatchingInteractionState time to the current player node's asset time. This is done 
 			// because the player node may have ticked more or less time than we expected due to variable dt or the 
-			// dynamic playback rate adjustment and as such the motion matching state does not update by itself
+			// dynamic playback rate adjustment and as such the MotionMatchingInteractionState does not update by itself
 			const FTimelineState TimelineState = TimelineTrait.GetState(Context);
 			float TimelinePosition = TimelineState.GetPosition();
 
 			// @TODO: This is a hack, since we don't have a way to get the normalized timeline of a blendspace yet (pending syncing...?)
 			if (const UE::PoseSearch::FSearchIndexAsset* SearchIndexAsset = MotionMatchingState.CurrentSearchResult.GetSearchIndexAsset())
 			{
-				if (const FPoseSearchDatabaseAnimationAssetBase* DatabaseAnimationAssetBase = MotionMatchingState.CurrentSearchResult.Database->GetDatabaseAnimationAsset<FPoseSearchDatabaseAnimationAssetBase>(*SearchIndexAsset))
+				if (const UMultiAnimAsset* MultiAnimAsset = Cast<UMultiAnimAsset>(MotionMatchingInteractionState.SelectedAnim))
 				{
-					const UObject* AnimationAsset = DatabaseAnimationAssetBase->GetAnimationAssetForRole(MotionMatchingState.CurrentSearchResult.Database->Schema->GetDefaultRole());
-					if (const UBlendSpace* BlendSpace = Cast<UBlendSpace>(AnimationAsset))
+					if (Cast<UBlendSpace>(MultiAnimAsset->GetAnimationAsset(MotionMatchingState.CurrentSearchResult.Database->Schema->GetDefaultRole())))
 					{
 						// Convert to normalized time.
 						const float TimelineDuration = TimelineState.GetDuration();
@@ -166,36 +163,10 @@ namespace UE::AnimNext
 				}
 			}
 
-			MotionMatchingState.AdjustAssetTime(TimelinePosition);
-			MotionMatchingInteractionState.SelectedTime = TimelinePosition;
-		}
-
-		const UObject* AnimationAsset = nullptr;
-		float AssetTime = 0.f;
-		float WantedPlayRate = 1.f;
-		bool bLoop = false;
-		FVector BlendSpaceParameters = FVector::ZeroVector;
-
-		const UObject* AnimContext = nullptr;
-		if (const FAnimNextModuleInstance* ModuleInstance = Context.GetRootGraphInstance().GetModuleInstance())
-		{
-			AnimContext = ModuleInstance->GetObject();
-		}
-
-		const TArray<FPoseSearchInteractionAvailability>& Availabilities = SharedData->GetAvailabilities(Binding);
-		if (Availabilities.IsEmpty())
-		{
-			MotionMatchingInteractionState = FPoseSearchBlueprintResult();
-		}
-		else
-		{
-			// WIP motion matched interactions
-			CheckInteractionThreadSafety(AnimContext);
-
 			// we're NOT providing the UMultiAnimAsset here, relying on the interaction subsystem to resolve the ContinuingProperties properly
 			FPoseSearchContinuingProperties ContinuingProperties;
 			ContinuingProperties.PlayingAsset = MotionMatchingInteractionState.SelectedAnim.Get(); // GetAnimAsset();
-			ContinuingProperties.PlayingAssetAccumulatedTime = MotionMatchingInteractionState.SelectedTime;
+			ContinuingProperties.PlayingAssetAccumulatedTime = TimelinePosition;
 
 			UPoseSearchInteractionLibrary::MotionMatchInteraction(MotionMatchingInteractionState, Availabilities, AnimContext, ContinuingProperties, FName(), PoseHistory, SharedData->GetbValidateResultAgainstAvailabilities(Binding));
 			check(MotionMatchingInteractionState.ActorRootTransforms.Num() == MotionMatchingInteractionState.ActorRootBoneTransforms.Num());
@@ -209,10 +180,6 @@ namespace UE::AnimNext
 				{
 					AnimationAsset = MultiAnimAsset->GetAnimationAsset(MotionMatchingInteractionState.Role);
 				}
-				else if (const UAnimationAsset* SingleAnimAsset = Cast<UAnimationAsset>(MotionMatchingInteractionState.SelectedAnim))
-				{
-					AnimationAsset = SingleAnimAsset;
-				}
 			}
 			AssetTime = MotionMatchingInteractionState.SelectedTime;
 			WantedPlayRate = MotionMatchingInteractionState.WantedPlayRate;
@@ -223,6 +190,68 @@ namespace UE::AnimNext
 		}
 		else
 		{
+			// synchronizing with GetAccumulatedTime or resetting MotionMatchingState, and conditionally resetting FAnimNode_BlendStack_Standalone
+			// @todo: implement this MotionMatchingState.Reset() condition for parity to FAnimNode_MotionMatching::UpdateAssetPlayer
+			//if (bResetOnBecomingRelevant && UpdateCounter.HasEverBeenUpdated() && !UpdateCounter.WasSynchronizedCounter(Context.AnimInstanceProxy->GetUpdateCounter()))
+			//{
+			//	// If we just became relevant and haven't been initialized yet, then reset motion matching state, otherwise update the asset time using the player node.
+			//	MotionMatchingState.Reset();
+			//	FAnimNode_BlendStack_Standalone::Reset();
+			//}
+			//else 
+			if (MotionMatchingState.CurrentSearchResult.Database == nullptr || MotionMatchingState.CurrentSearchResult.Database->Schema == nullptr)
+			{
+			}
+	#if WITH_EDITOR
+			else if (EAsyncBuildIndexResult::Success != FAsyncPoseSearchDatabasesManagement::RequestAsyncBuildIndex(MotionMatchingState.CurrentSearchResult.Database.Get(), ERequestAsyncBuildFlag::ContinueRequest))
+			{
+				// MotionMatchingState.CurrentSearchResult.Database is indexing, and it's not safe to use its previous index properties cached in MotionMatchingState
+				MotionMatchingState.Reset();
+			}
+			// in case this node is not updated, and MotionMatchingState.CurrentSearchResult.Database gets modified, we could end up with CurrentSearchResult being out of synch with the updated database, so we need to reset the state
+			else if (MotionMatchingState.CurrentSearchResult.IsValid() && MotionMatchingState.CurrentSearchResult.PoseIdx >= MotionMatchingState.CurrentSearchResult.Database->GetSearchIndex().GetNumPoses())
+			{
+				// in case this node is not updated, and MotionMatchingState.CurrentSearchResult.Database gets modified, we could end up with
+				// CurrentSearchResult being out of synch with the updated database, so we need to reset the state
+				MotionMatchingState.Reset();
+			}
+	#endif // WITH_EDITOR
+			// @todo: implement this MotionMatchingState.Reset() condition for parity to FAnimNode_MotionMatching::UpdateAssetPlayer
+			//else if (GetAnimAsset() != MotionMatchingState.CurrentSearchResult.GetCurrentResultAnimationAsset())
+			//{
+			//	// MotionMatchingState and blendstack are not in synch! MotionMatchingState thinks it's playing MotionMatchingState.CurrentSearchResult.GetCurrentResultAnimationAsset()
+			//	// and blendstack thinks it's playing GetAnimAsset(). The two could have gone out of synch becasue or external blendstack requests, and so we reset the MotionMatchingState
+			//	MotionMatchingState.Reset();
+			//}
+			else
+			{
+				// We adjust the MotionMatchingState asset time to the current player node's asset time. This is done 
+				// because the player node may have ticked more or less time than we expected due to variable dt or the 
+				// dynamic playback rate adjustment and as such the MotionMatchingState does not update by itself
+				const FTimelineState TimelineState = TimelineTrait.GetState(Context);
+				float TimelinePosition = TimelineState.GetPosition();
+
+				// @TODO: This is a hack, since we don't have a way to get the normalized timeline of a blendspace yet (pending syncing...?)
+				if (const UE::PoseSearch::FSearchIndexAsset* SearchIndexAsset = MotionMatchingState.CurrentSearchResult.GetSearchIndexAsset())
+				{
+					if (const FPoseSearchDatabaseAnimationAssetBase* DatabaseAnimationAssetBase = MotionMatchingState.CurrentSearchResult.Database->GetDatabaseAnimationAsset<FPoseSearchDatabaseAnimationAssetBase>(*SearchIndexAsset))
+					{
+						if (Cast<UBlendSpace>(DatabaseAnimationAssetBase->GetAnimationAssetForRole(MotionMatchingState.CurrentSearchResult.Database->Schema->GetDefaultRole())))
+						{
+							// Convert to normalized time.
+							const float TimelineDuration = TimelineState.GetDuration();
+							if (TimelineDuration > 0.0f)
+							{
+								// NOTE: This doesn't work because AnimNext normalized time does not much database indexing normalized time.
+								TimelinePosition /= TimelineState.GetDuration();
+							}
+						}
+					}
+				}
+
+				MotionMatchingState.AdjustAssetTime(TimelinePosition);
+			}
+
 			const TConstArrayView<TObjectPtr<const UPoseSearchDatabase>> DatabasesToSearch = SharedData->GetDatabases(Binding);
 			if (DatabasesToSearch.IsEmpty())
 			{
