Anim Next: Only use synchronization on blendspaces and not sequences too.

diff --git a/Engine/Plugins/Experimental/AnimNextPoseSearch/Source/AnimNextPoseSearch/Private/MotionMatchingTrait.cpp b/Engine/Plugins/Experimental/AnimNextPoseSearch/Source/AnimNextPoseSearch/Private/MotionMatchingTrait.cpp
--- a/Engine/Plugins/Experimental/AnimNextPoseSearch/Source/AnimNextPoseSearch/Private/MotionMatchingTrait.cpp
+++ b/Engine/Plugins/Experimental/AnimNextPoseSearch/Source/AnimNextPoseSearch/Private/MotionMatchingTrait.cpp
@@ -11,6 +11,7 @@
 #include "Graph/AnimNextGraphInstance.h"
 #include "GraphInterfaces/AnimNextNativeDataInterface_AnimSequencePlayer.h"
 #include "GraphInterfaces/AnimNextNativeDataInterface_BlendSpacePlayer.h"
+#include "GraphInterfaces/AnimNextNativeDataInterface_SynchronizeUsingGroups.h"
 #include "IPoseHistory.h"
 #include "Module/AnimNextModuleInstance.h"
 #include "Module/RigUnit_AnimNextModuleEvents.h"
@@ -341,11 +342,21 @@ namespace UE::AnimNext
 				NewGraphRequest.BlendArgs = SharedData->BlendArguments;
 
 				// Create a default 'native data interface' payload via which we expect to communicate with the graph 
-				FInstancedStruct NewPayload;
+				FAnimNextDataInterfacePayload GraphPayload;
+				FInstancedStruct AssetPlayerPayload;
 				// @TODO: Support asset->playload mapping in a more data-driven way
 				if (bIsBlendSpace)
 				{
-					FAnimNextNativeDataInterface_BlendSpacePlayer& NewNativeInterface = NewPayload.InitializeAs<FAnimNextNativeDataInterface_BlendSpacePlayer>();
+					// Tell blendspace samples to synchronize among themselves, but not with anything else.
+					FInstancedStruct SyncGroupPayload;
+					auto& NativeInterface = SyncGroupPayload.InitializeAs<FAnimNextNativeDataInterface_SynchronizeUsingGroups>();
+					NativeInterface.GroupName = NAME_None;
+					NativeInterface.GroupRole = EAnimGroupSynchronizationRole::AlwaysFollower;
+					NativeInterface.SyncMode = EAnimGroupSynchronizationMode::SynchronizeUsingUniqueGroupName;
+					NativeInterface.MatchSyncPoint = true;
+					GraphPayload.AddNative(MoveTemp(SyncGroupPayload));
+
+					FAnimNextNativeDataInterface_BlendSpacePlayer& NewNativeInterface = AssetPlayerPayload.InitializeAs<FAnimNextNativeDataInterface_BlendSpacePlayer>();
 					NewNativeInterface.BlendSpace = Cast<UBlendSpace>(AnimationAsset);
 
 					NewNativeInterface.XAxisSamplePoint = BlendSpaceParameters.X;
@@ -357,15 +368,14 @@ namespace UE::AnimNext
 				}
 				else
 				{
-					FAnimNextNativeDataInterface_AnimSequencePlayer& NewNativeInterface = NewPayload.InitializeAs<FAnimNextNativeDataInterface_AnimSequencePlayer>();
+					FAnimNextNativeDataInterface_AnimSequencePlayer& NewNativeInterface = AssetPlayerPayload.InitializeAs<FAnimNextNativeDataInterface_AnimSequencePlayer>();
 					NewNativeInterface.AnimSequence = Cast<UAnimSequence>(AnimationAsset);
 					NewNativeInterface.PlayRate = DesiredPlayRate;
 					NewNativeInterface.StartPosition = AssetTime;
 					NewNativeInterface.Loop = bLoop;
 				}
 
-				FAnimNextDataInterfacePayload GraphPayload;
-				GraphPayload.AddNative(MoveTemp(NewPayload));
+				GraphPayload.AddNative(MoveTemp(AssetPlayerPayload));
 
 				const UAnimNextAnimationGraph* AnimationGraph = IGraphFactory::GetGraphFromObjectWithFallback(Context, Binding, AnimationAsset, GraphPayload);
 				if (AnimationGraph != nullptr)
