Fix for shrinking in motion matching

diff --git a/Engine/Plugins/Experimental/AnimNextPoseSearch/Source/AnimNextPoseSearch/Private/MotionMatchingTrait.cpp b/Engine/Plugins/Experimental/AnimNextPoseSearch/Source/AnimNextPoseSearch/Private/MotionMatchingTrait.cpp
--- a/Engine/Plugins/Experimental/AnimNextPoseSearch/Source/AnimNextPoseSearch/Private/MotionMatchingTrait.cpp
+++ b/Engine/Plugins/Experimental/AnimNextPoseSearch/Source/AnimNextPoseSearch/Private/MotionMatchingTrait.cpp
@@ -73,7 +73,6 @@ namespace UE::AnimNext
 		using namespace UE::PoseSearch;
 
 		FInstanceData* InstanceData = Binding.GetInstanceData<FInstanceData>();
-		IUpdate::PreUpdate(Context, Binding, TraitState);
 
 #if WITH_EDITOR
 		if (InstanceData->bIsPostEvaluateBeingCalled)
@@ -417,6 +416,9 @@ namespace UE::AnimNext
 		}
 
 		PublishResults(Binding);
+
+		// Motion Matching must finish updating before smoothers, otherwise we may request a blend after the smoother is done normalizing weights.
+		IUpdate::PreUpdate(Context, Binding, TraitState);
 	}
 
 	void FMotionMatchingTrait::PublishResults(const TTraitBinding<IUpdate>& Binding) const
