Add Pose Search Interaction types to RigVM Registry

diff --git a/Engine/Plugins/Experimental/AnimNextPoseSearch/Source/AnimNextPoseSearch/Private/Module.cpp b/Engine/Plugins/Experimental/AnimNextPoseSearch/Source/AnimNextPoseSearch/Private/Module.cpp
--- a/Engine/Plugins/Experimental/AnimNextPoseSearch/Source/AnimNextPoseSearch/Private/Module.cpp
+++ b/Engine/Plugins/Experimental/AnimNextPoseSearch/Source/AnimNextPoseSearch/Private/Module.cpp
@@ -10,6 +10,7 @@
 #include "AnimNextInteractionIslandDependency.h"
 #include "Features/IModularFeatures.h"
 #include "PoseSearch/PoseSearchDatabase.h"
+#include "PoseSearch/PoseSearchInteractionAsset.h"
 #include "PoseSearch/PoseSearchInteractionLibrary.h"
 #include "PoseSearch/PoseSearchSchema.h"
 #include "Animation/TrajectoryTypes.h"
@@ -25,7 +26,9 @@ public:
 		static TPair<UClass*, FRigVMRegistry::ERegisterObjectOperation> const AllowedObjectTypes[] =
 		{
 			{ UPoseSearchDatabase::StaticClass(), FRigVMRegistry::ERegisterObjectOperation::Class },
-			{ UPoseSearchSchema::StaticClass(), FRigVMRegistry::ERegisterObjectOperation::Class }
+			{ UPoseSearchSchema::StaticClass(), FRigVMRegistry::ERegisterObjectOperation::Class },
+			{ UPoseSearchInteractionAsset::StaticClass(), FRigVMRegistry::ERegisterObjectOperation::Class },
+			{ UMultiAnimAsset::StaticClass(), FRigVMRegistry::ERegisterObjectOperation::Class }
 		};
 
 		static UScriptStruct* const AllowedStructTypes[] =
@@ -34,6 +37,7 @@ public:
 			FTransformTrajectory::StaticStruct(),
 			FPoseSearchBlueprintResult::StaticStruct(),
 			FAlphaBlendArgs::StaticStruct(),
+			FPoseSearchInteractionAssetItem::StaticStruct(),
 			FPoseSearchInteractionAvailability::StaticStruct(),
 		};
 
