Anim Next Pose Search: Special blend stack provider that attempts to emulate the result of a Pose Search

diff --git a/Engine/Plugins/Experimental/AnimNextPoseSearch/Source/AnimNextPoseSearch/Private/PoseSearchResultEmulatorTrait.cpp b/Engine/Plugins/Experimental/AnimNextPoseSearch/Source/AnimNextPoseSearch/Private/PoseSearchResultEmulatorTrait.cpp
new file mode 100644
--- /dev/null
+++ b/Engine/Plugins/Experimental/AnimNextPoseSearch/Source/AnimNextPoseSearch/Private/PoseSearchResultEmulatorTrait.cpp
@@ -0,0 +1,215 @@
+// Copyright Epic Games, Inc. All Rights Reserved.
+
+#include "PoseSearchResultEmulatorTrait.h"
+#include "Animation/AnimRootMotionProvider.h"
+#include "Animation/AnimSequence.h"
+#include "Animation/BlendSpace.h"
+#include "DataInterface/DataInterfaceStructAdapter.h"
+#include "Graph/AnimNextGraphInstance.h"
+#include "GraphInterfaces/AnimNextNativeDataInterface_AnimSequencePlayer.h"
+#include "GraphInterfaces/AnimNextNativeDataInterface_BlendSpacePlayer.h"
+#include "PoseSearch/MultiAnimAsset.h"
+#include "TraitInterfaces/IBlendStack.h"
+#include "TraitInterfaces/IGraphFactory.h"
+#include "TraitInterfaces/ITimeline.h"
+#include "EvaluationVM/EvaluationVM.h"
+#include "GraphInterfaces/AnimNextNativeDataInterface_SynchronizeUsingGroups.h"
+
+namespace UE::AnimNext
+{
+	AUTO_REGISTER_ANIM_TRAIT(FPoseSearchResultEmulatorTrait)
+
+#define TRAIT_INTERFACE_ENUMERATOR(GeneratorMacro) \
+		GeneratorMacro(IUpdate) \
+		GeneratorMacro(IEvaluate) \
+
+	// Trait implementation boilerplate
+	GENERATE_ANIM_TRAIT_IMPLEMENTATION(FPoseSearchResultEmulatorTrait, TRAIT_INTERFACE_ENUMERATOR, NULL_ANIM_TRAIT_INTERFACE_ENUMERATOR, NULL_ANIM_TRAIT_EVENT_ENUMERATOR)
+#undef TRAIT_INTERFACE_ENUMERATOR
+
+	void FPoseSearchResultEmulatorTrait::PreUpdate(FUpdateTraversalContext& Context, const TTraitBinding<IUpdate>& Binding, const FTraitUpdateState& TraitState) const
+	{
+		using namespace UE::PoseSearch;
+		FInstanceData* InstanceData = Binding.GetInstanceData<FInstanceData>();
+
+		TTraitBinding<IBlendStack> BlendStackBinding;
+		if (!Binding.GetStackInterface(BlendStackBinding))
+		{
+			UE_LOG(LogPoseSearch, Error, TEXT("FPoseSearchResultEmulatorTrait::PreUpdate, missing IBlendStack"));
+			return;
+		}
+
+		TTraitBinding<ITimeline> TimelineTrait;
+		if (!Binding.GetStackInterface(TimelineTrait))
+		{
+			UE_LOG(LogPoseSearch, Error, TEXT("FPoseSearchResultEmulatorTrait::PreUpdate, missing ITimeline"));
+			return;
+		}
+
+		const FSharedData* SharedData = Binding.GetSharedData<FSharedData>();
+		const UObject* AnimObject = SharedData->GetSelectedAnim(Binding);
+		if (AnimObject == nullptr)
+		{
+			// @TODO: Kill this with fire.
+			const FPoseSearchBlueprintResult& Result = SharedData->GetPoseSearchResult(Binding);
+			AnimObject = Result.SelectedAnim.Get();
+		}
+
+		if (AnimObject)
+		{
+			// Shared parameters
+			const FName Role = SharedData->GetRole(Binding);
+			const float SelectedTime = SharedData->GetSelectedTime(Binding);
+			const float WantedPlayRate = SharedData->GetWantedPlayRate(Binding);
+			const bool bLoop = SharedData->GetbLoop(Binding);
+			const FVector& BlendParameters = SharedData->GetBlendParameters(Binding);
+
+			const UAnimationAsset* AnimationAsset = nullptr;
+			if (const UMultiAnimAsset* MultiAnimAsset = Cast<UMultiAnimAsset>(AnimObject))
+			{
+				AnimationAsset = MultiAnimAsset->GetAnimationAsset(Role);
+			}
+			else
+			{
+				AnimationAsset = Cast<UAnimationAsset>(AnimObject);
+			}
+
+			IBlendStack::FGraphRequestPtr ActiveGraphRequest;
+			BlendStackBinding.GetActiveGraph(Context, ActiveGraphRequest);
+			const UObject* CurrentAnim = nullptr;
+			float TimeDelta = 0.0f;
+			bool bIsBlendSpace = false;
+			const bool bHasValidActiveGraph = ActiveGraphRequest && ActiveGraphRequest->FactoryObject;
+			if (bHasValidActiveGraph)
+			{
+				CurrentAnim = ActiveGraphRequest->FactoryObject;
+				const FTimelineState TimelineState = TimelineTrait.GetState(Context);
+				float TimelinePosition = TimelineState.GetPosition();
+
+				bIsBlendSpace = CurrentAnim->IsA<UBlendSpace>();
+				if (bIsBlendSpace)
+				{
+					// Convert to normalized time.
+					const float TimelineDuration = TimelineState.GetDuration();
+					if (TimelineDuration > 0.0f)
+					{
+						// NOTE: This doesn't work because AnimNext normalized time does not match database indexing normalized time.
+						TimelinePosition /= TimelineState.GetDuration();
+					}
+				}
+
+				TimeDelta = FMath::Abs(TimelinePosition - SelectedTime);
+			}
+
+			const float MaxTimeDeltaAllowed = SharedData->GetMaxTimeDeltaAllowed(Binding);
+			const bool bDoBlend =	CurrentAnim != AnimationAsset
+								||	(TimeDelta > MaxTimeDeltaAllowed && !bIsBlendSpace);
+
+			if (bDoBlend)
+			{
+				UE::AnimNext::IBlendStack::FGraphRequest NewGraphRequest;
+				NewGraphRequest.BlendArgs = SharedData->GetBlendArguments(Binding);
+
+				// Create a default 'native data interface' payload via which we expect to communicate with the graph 
+				FAnimNextDataInterfacePayload GraphPayload;
+				FInstancedStruct AssetPlayerPayload;
+				if (bIsBlendSpace)
+				{
+					// Tell blendspace samples to synchronize among themselves, but not with anything else.
+					FInstancedStruct SyncGroupPayload;
+					auto& NativeInterface = SyncGroupPayload.InitializeAs<FAnimNextNativeDataInterface_SynchronizeUsingGroups>();
+					NativeInterface.GroupName = NAME_None;
+					NativeInterface.GroupRole = EAnimGroupSynchronizationRole::AlwaysFollower;
+					NativeInterface.SyncMode = EAnimGroupSynchronizationMode::SynchronizeUsingUniqueGroupName;
+					NativeInterface.MatchSyncPoint = true;
+					GraphPayload.AddNative(MoveTemp(SyncGroupPayload));
+
+					FAnimNextNativeDataInterface_BlendSpacePlayer& NewNativeInterface = AssetPlayerPayload.InitializeAs<FAnimNextNativeDataInterface_BlendSpacePlayer>();
+					NewNativeInterface.BlendSpace = Cast<UBlendSpace>(AnimationAsset);
+
+					NewNativeInterface.XAxisSamplePoint = BlendParameters.X;
+					NewNativeInterface.YAxisSamplePoint = BlendParameters.Y;
+
+					NewNativeInterface.PlayRate = WantedPlayRate;
+					NewNativeInterface.StartPosition =  SelectedTime;
+					NewNativeInterface.Loop = bLoop;
+				}
+				else
+				{
+					FAnimNextNativeDataInterface_AnimSequencePlayer& NewNativeInterface = AssetPlayerPayload.InitializeAs<FAnimNextNativeDataInterface_AnimSequencePlayer>();
+					NewNativeInterface.AnimSequence = Cast<UAnimSequence>(AnimationAsset);
+					NewNativeInterface.PlayRate = WantedPlayRate;
+					NewNativeInterface.StartPosition =  SelectedTime;
+					NewNativeInterface.Loop = bLoop;
+				}
+
+				GraphPayload.AddNative(MoveTemp(AssetPlayerPayload));
+
+				const UAnimNextAnimationGraph* AnimationGraph = IGraphFactory::GetGraphFromObjectWithFallback(Context, Binding, AnimationAsset, GraphPayload);
+				if (AnimationGraph != nullptr)
+				{
+					NewGraphRequest.GraphPayload = MoveTemp(GraphPayload);
+					NewGraphRequest.AnimationGraph = AnimationGraph;
+					NewGraphRequest.FactoryObject = AnimationAsset;
+					BlendStackBinding.PushGraph(Context, MoveTemp(NewGraphRequest));
+				}
+			}
+
+			if (bHasValidActiveGraph)
+			{
+				for (FStructView& Payload : ActiveGraphRequest->GraphPayload.Get())
+				{
+					if (FAnimNextNativeDataInterface_AnimSequencePlayer* PlayAnimPayload = Payload.GetPtr<FAnimNextNativeDataInterface_AnimSequencePlayer>())
+					{
+						PlayAnimPayload->PlayRate = WantedPlayRate;
+						break;
+					}
+					else if (FAnimNextNativeDataInterface_BlendSpacePlayer* BlendSpacePayload = Payload.GetPtr<FAnimNextNativeDataInterface_BlendSpacePlayer>())
+					{
+						BlendSpacePayload->PlayRate = WantedPlayRate;
+						BlendSpacePayload->XAxisSamplePoint = BlendParameters.X;
+						BlendSpacePayload->YAxisSamplePoint = BlendParameters.Y;
+						break;
+					}
+				}
+			}
+		}
+
+		IUpdate::PreUpdate(Context, Binding, TraitState);
+	}
+
+	void FPoseSearchResultEmulatorTrait::PostEvaluate(FEvaluateTraversalContext& Context, const TTraitBinding<IEvaluate>& Binding) const
+	{
+		IEvaluate::PostEvaluate(Context, Binding);
+
+		const FSharedData* SharedData = Binding.GetSharedData<FSharedData>();
+		const bool bDisableRootMotion = SharedData->GetbDisableRootMotion(Binding);
+
+		if (bDisableRootMotion)
+		{
+			FAnimNextPoseSearchResultEmulatorTask Task;
+			Context.AppendTask(Task);
+		}
+	}
+	
+} // UE::AnimNext
+
+void FAnimNextPoseSearchResultEmulatorTask::Execute(UE::AnimNext::FEvaluationVM& VM) const
+{
+	using namespace UE::AnimNext;
+
+	const UE::Anim::IAnimRootMotionProvider* RootMotionProvider = UE::Anim::IAnimRootMotionProvider::Get();
+	if (!RootMotionProvider)
+	{
+		UE_LOG(LogPoseSearch, Error, TEXT("FAnimNextPoseSearchResultEmulatorTask::Execute, missing RootMotionProvider"));
+		return;
+	}
+
+	const TUniquePtr<FKeyframeState>* Keyframe = VM.PeekValue<TUniquePtr<FKeyframeState>>(KEYFRAME_STACK_NAME, 0);
+	if (!Keyframe)
+	{
+		return;
+	}
+
+	RootMotionProvider->OverrideRootMotion(FTransform::Identity, Keyframe->Get()->Attributes);
+}
diff --git a/Engine/Plugins/Experimental/AnimNextPoseSearch/Source/AnimNextPoseSearch/Private/PoseSearchResultEmulatorTrait.h b/Engine/Plugins/Experimental/AnimNextPoseSearch/Source/AnimNextPoseSearch/Private/PoseSearchResultEmulatorTrait.h
new file mode 100644
--- /dev/null
+++ b/Engine/Plugins/Experimental/AnimNextPoseSearch/Source/AnimNextPoseSearch/Private/PoseSearchResultEmulatorTrait.h
@@ -0,0 +1,96 @@
+// Copyright Epic Games, Inc. All Rights Reserved.
+
+#pragma once
+
+#include "CoreMinimal.h"
+#include "PoseSearch/PoseSearchLibrary.h"
+#include "TraitCore/Trait.h"
+#include "TraitCore/TraitSharedData.h"
+#include "TraitInterfaces/IUpdate.h"
+#include "TraitInterfaces/IEvaluate.h"
+#include "PoseSearchResultEmulatorTrait.generated.h"
+
+USTRUCT(meta = (DisplayName = "Pose Search Result Emulator"))
+struct FPoseSearchResultEmulatorTraitSharedData : public FAnimNextTraitSharedData
+{
+	GENERATED_BODY()
+
+	UPROPERTY(EditAnywhere, Category = Default)
+	TObjectPtr<UObject> SelectedAnim = nullptr;
+
+	// @TODO: This is a super hack, but we don't have a good way to use TWeakObjectPtr pose search result. Do this for now.
+	UPROPERTY(EditAnywhere, Category = Default)
+	FPoseSearchBlueprintResult PoseSearchResult;
+	
+	UPROPERTY(EditAnywhere, Category = Default)
+	float SelectedTime = 0.f;
+	
+	UPROPERTY(EditAnywhere, Category = Default)
+	float WantedPlayRate = 0.f;
+
+	UPROPERTY(EditAnywhere, Category = Default)
+	bool bLoop = false;
+	
+	UPROPERTY(EditAnywhere, Category = Default)
+	FVector BlendParameters = FVector::ZeroVector;
+
+	UPROPERTY(EditAnywhere, Category = Default)
+	FName Role;
+
+	UPROPERTY(EditAnywhere, Category = Default)
+	FAlphaBlendArgs BlendArguments;
+
+	UPROPERTY(EditAnywhere, Category = Default)
+	float MaxTimeDeltaAllowed = 0.1f;
+
+	UPROPERTY(EditAnywhere, Category = Default)
+	bool bDisableRootMotion = false;
+
+	// Latent pin support boilerplate
+	#define TRAIT_LATENT_PROPERTIES_ENUMERATOR(GeneratorMacro) \
+		GeneratorMacro(SelectedAnim) \
+		GeneratorMacro(PoseSearchResult) \
+		GeneratorMacro(SelectedTime) \
+		GeneratorMacro(WantedPlayRate) \
+		GeneratorMacro(bLoop) \
+		GeneratorMacro(BlendParameters) \
+		GeneratorMacro(Role) \
+		GeneratorMacro(BlendArguments) \
+		GeneratorMacro(MaxTimeDeltaAllowed) \
+		GeneratorMacro(bDisableRootMotion) \
+
+	GENERATE_TRAIT_LATENT_PROPERTIES(FPoseSearchResultEmulatorTraitSharedData, TRAIT_LATENT_PROPERTIES_ENUMERATOR)
+	#undef TRAIT_LATENT_PROPERTIES_ENUMERATOR
+};
+
+namespace UE::AnimNext
+{
+
+struct FPoseSearchResultEmulatorTrait : FAdditiveTrait, IUpdate, IEvaluate
+{
+	DECLARE_ANIM_TRAIT(FPoseSearchResultEmulatorTrait, FAdditiveTrait)
+
+	using FSharedData = FPoseSearchResultEmulatorTraitSharedData;
+		
+	struct FInstanceData : FTrait::FInstanceData
+	{
+	};
+
+	// IUpdate impl
+	virtual void PreUpdate(FUpdateTraversalContext& Context, const TTraitBinding<IUpdate>& Binding, const FTraitUpdateState& TraitState) const override;
+
+	// IEvaluate impl
+	virtual void PostEvaluate(FEvaluateTraversalContext& Context, const TTraitBinding<IEvaluate>& Binding) const override;
+};
+
+} // namespace UE::AnimNext
+
+USTRUCT(Experimental)
+struct FAnimNextPoseSearchResultEmulatorTask : public FAnimNextEvaluationTask
+{
+	GENERATED_BODY()
+
+	DECLARE_ANIM_EVALUATION_TASK(FAnimNextPoseSearchResultEmulatorTask)
+
+	virtual void Execute(UE::AnimNext::FEvaluationVM& VM) const override;
+};
\ No newline at end of file
