Speculative fix for crash in PoseSearchResultEmulatorTrait

diff --git a/Engine/Plugins/Experimental/AnimNextPoseSearch/Source/AnimNextPoseSearch/Private/PoseSearchResultEmulatorTrait.cpp b/Engine/Plugins/Experimental/AnimNextPoseSearch/Source/AnimNextPoseSearch/Private/PoseSearchResultEmulatorTrait.cpp
--- a/Engine/Plugins/Experimental/AnimNextPoseSearch/Source/AnimNextPoseSearch/Private/PoseSearchResultEmulatorTrait.cpp
+++ b/Engine/Plugins/Experimental/AnimNextPoseSearch/Source/AnimNextPoseSearch/Private/PoseSearchResultEmulatorTrait.cpp
@@ -154,8 +154,10 @@ namespace UE::AnimNext
 					BlendStackBinding.PushGraph(Context, MoveTemp(NewGraphRequest));
 				}
 			}
-
-			if (bHasValidActiveGraph)
+			
+			// HACK: Grab the active request again. It seems the previous active graph request can become invalid here.
+			BlendStackBinding.GetActiveGraph(Context, ActiveGraphRequest);
+			if (ActiveGraphRequest)
 			{
 				for (FStructView& Payload : ActiveGraphRequest->GraphPayload.Get())
 				{
