Change pose search emulator trait input pins to work around FVector latent pin bug

diff --git a/Engine/Plugins/Experimental/AnimNextPoseSearch/Source/AnimNextPoseSearch/Private/PoseSearchResultEmulatorTrait.cpp b/Engine/Plugins/Experimental/AnimNextPoseSearch/Source/AnimNextPoseSearch/Private/PoseSearchResultEmulatorTrait.cpp
--- a/Engine/Plugins/Experimental/AnimNextPoseSearch/Source/AnimNextPoseSearch/Private/PoseSearchResultEmulatorTrait.cpp
+++ b/Engine/Plugins/Experimental/AnimNextPoseSearch/Source/AnimNextPoseSearch/Private/PoseSearchResultEmulatorTrait.cpp
@@ -62,7 +62,8 @@ namespace UE::AnimNext
 			const float SelectedTime = SharedData->GetSelectedTime(Binding);
 			const float WantedPlayRate = SharedData->GetWantedPlayRate(Binding);
 			const bool bLoop = SharedData->GetbLoop(Binding);
-			const FVector& BlendParameters = SharedData->GetBlendParameters(Binding);
+			const float XAxisSamplePoint = SharedData->GetXAxisSamplePoint(Binding);
+			const float YAxisSamplePoint = SharedData->GetYAxisSamplePoint(Binding);
 
 			const UAnimationAsset* AnimationAsset = nullptr;
 			if (const UMultiAnimAsset* MultiAnimAsset = Cast<UMultiAnimAsset>(AnimObject))
@@ -78,7 +79,7 @@ namespace UE::AnimNext
 			BlendStackBinding.GetActiveGraph(Context, ActiveGraphRequest);
 			const UObject* CurrentAnim = nullptr;
 			float TimeDelta = 0.0f;
-			bool bIsBlendSpace = false;
+			bool bIsBlendSpace = AnimationAsset && AnimationAsset->IsA<UBlendSpace>();
 			const bool bHasValidActiveGraph = ActiveGraphRequest && ActiveGraphRequest->FactoryObject;
 			if (bHasValidActiveGraph)
 			{
@@ -86,8 +87,8 @@ namespace UE::AnimNext
 				const FTimelineState TimelineState = TimelineTrait.GetState(Context);
 				float TimelinePosition = TimelineState.GetPosition();
 
-				bIsBlendSpace = CurrentAnim->IsA<UBlendSpace>();
-				if (bIsBlendSpace)
+				const bool bCurrentAnimIsBlendSpace = CurrentAnim->IsA<UBlendSpace>();
+				if (bCurrentAnimIsBlendSpace)
 				{
 					// Convert to normalized time.
 					const float TimelineDuration = TimelineState.GetDuration();
@@ -127,8 +128,8 @@ namespace UE::AnimNext
 					FAnimNextNativeDataInterface_BlendSpacePlayer& NewNativeInterface = AssetPlayerPayload.InitializeAs<FAnimNextNativeDataInterface_BlendSpacePlayer>();
 					NewNativeInterface.BlendSpace = Cast<UBlendSpace>(AnimationAsset);
 
-					NewNativeInterface.XAxisSamplePoint = BlendParameters.X;
-					NewNativeInterface.YAxisSamplePoint = BlendParameters.Y;
+					NewNativeInterface.XAxisSamplePoint = XAxisSamplePoint;
+					NewNativeInterface.YAxisSamplePoint = YAxisSamplePoint;
 
 					NewNativeInterface.PlayRate = WantedPlayRate;
 					NewNativeInterface.StartPosition =  SelectedTime;
@@ -169,8 +170,8 @@ namespace UE::AnimNext
 					else if (FAnimNextNativeDataInterface_BlendSpacePlayer* BlendSpacePayload = Payload.GetPtr<FAnimNextNativeDataInterface_BlendSpacePlayer>())
 					{
 						BlendSpacePayload->PlayRate = WantedPlayRate;
-						BlendSpacePayload->XAxisSamplePoint = BlendParameters.X;
-						BlendSpacePayload->YAxisSamplePoint = BlendParameters.Y;
+						BlendSpacePayload->XAxisSamplePoint = XAxisSamplePoint;
+						BlendSpacePayload->YAxisSamplePoint = YAxisSamplePoint;
 						break;
 					}
 				}
diff --git a/Engine/Plugins/Experimental/AnimNextPoseSearch/Source/AnimNextPoseSearch/Private/PoseSearchResultEmulatorTrait.h b/Engine/Plugins/Experimental/AnimNextPoseSearch/Source/AnimNextPoseSearch/Private/PoseSearchResultEmulatorTrait.h
--- a/Engine/Plugins/Experimental/AnimNextPoseSearch/Source/AnimNextPoseSearch/Private/PoseSearchResultEmulatorTrait.h
+++ b/Engine/Plugins/Experimental/AnimNextPoseSearch/Source/AnimNextPoseSearch/Private/PoseSearchResultEmulatorTrait.h
@@ -32,7 +32,10 @@ struct FPoseSearchResultEmulatorTraitSharedData : public FAnimNextTraitSharedDat
 	bool bLoop = false;
 	
 	UPROPERTY(EditAnywhere, Category = Default)
-	FVector BlendParameters = FVector::ZeroVector;
+	float XAxisSamplePoint = 0.0f;
+
+	UPROPERTY(EditAnywhere, Category = Default)
+	float YAxisSamplePoint = 0.0f;
 
 	UPROPERTY(EditAnywhere, Category = Default)
 	FName Role;
@@ -53,7 +56,8 @@ struct FPoseSearchResultEmulatorTraitSharedData : public FAnimNextTraitSharedDat
 		GeneratorMacro(SelectedTime) \
 		GeneratorMacro(WantedPlayRate) \
 		GeneratorMacro(bLoop) \
-		GeneratorMacro(BlendParameters) \
+		GeneratorMacro(XAxisSamplePoint) \
+		GeneratorMacro(YAxisSamplePoint) \
 		GeneratorMacro(Role) \
 		GeneratorMacro(BlendArguments) \
 		GeneratorMacro(MaxTimeDeltaAllowed) \
