Pose Search Distance Channel, which gets its data from an anim next variable

diff --git a/Engine/Plugins/Experimental/AnimNextPoseSearch/Source/AnimNextPoseSearch/Private/PoseSearchFeatureChannel_DistanceFromAnimNextVar.cpp b/Engine/Plugins/Experimental/AnimNextPoseSearch/Source/AnimNextPoseSearch/Private/PoseSearchFeatureChannel_DistanceFromAnimNextVar.cpp
new file mode 100644
--- /dev/null
+++ b/Engine/Plugins/Experimental/AnimNextPoseSearch/Source/AnimNextPoseSearch/Private/PoseSearchFeatureChannel_DistanceFromAnimNextVar.cpp
@@ -0,0 +1,26 @@
+// Copyright Epic Games, Inc. All Rights Reserved.
+
+#include "PoseSearchFeatureChannel_DistanceFromAnimNextVar.h"
+#include "Component/AnimNextComponent.h"
+
+#include "PoseSearch/PoseSearchContext.h"
+#include "StructUtils/PropertyBag.h"
+
+void UPoseSearchFeatureChannel_DistanceFromAnimNextVar::BuildQuery(UE::PoseSearch::FSearchContext& SearchContext) const
+{
+	using namespace UE::PoseSearch;
+	const UObject* Owner = SearchContext.GetAnimContext(SampleRole);
+	
+	float Distance = 0.0f;
+
+	if (const UAnimNextComponent* Component = Cast<UAnimNextComponent>(Owner))
+	{
+		TOptional<float> Result = Component->GetVariableFloat(DistanceVariableName);
+		if (Result.IsSet())
+		{
+			Distance = Result.GetValue();
+		}
+	}
+	
+	FFeatureVectorHelper::EncodeFloat(SearchContext.EditFeatureVector(), ChannelDataOffset, Distance);
+}
diff --git a/Engine/Plugins/Experimental/AnimNextPoseSearch/Source/AnimNextPoseSearch/Private/PoseSearchFeatureChannel_DistanceFromAnimNextVar.h b/Engine/Plugins/Experimental/AnimNextPoseSearch/Source/AnimNextPoseSearch/Private/PoseSearchFeatureChannel_DistanceFromAnimNextVar.h
new file mode 100644
--- /dev/null
+++ b/Engine/Plugins/Experimental/AnimNextPoseSearch/Source/AnimNextPoseSearch/Private/PoseSearchFeatureChannel_DistanceFromAnimNextVar.h
@@ -0,0 +1,17 @@
+// Copyright Epic Games, Inc. All Rights Reserved.
+
+#pragma once
+
+#include "PoseSearch/PoseSearchFeatureChannel_Distance.h"
+#include "PoseSearchFeatureChannel_DistanceFromAnimNextVar.generated.h"
+
+UCLASS(Experimental, EditInlineNew, Blueprintable, meta = (DisplayName = "Distance Channel (AnimNext Variable)"), CollapseCategories)
+class UPoseSearchFeatureChannel_DistanceFromAnimNextVar : public UPoseSearchFeatureChannel_Distance
+{
+	GENERATED_BODY()
+	
+	UPROPERTY(DisplayName = "Variable", EditAnywhere, Category="Variable")
+	FName DistanceVariableName;
+
+	virtual void BuildQuery(UE::PoseSearch::FSearchContext& SearchContext) const override;
+};
