Fix duplicate symbol in build

diff --git a/Engine/Plugins/Experimental/AnimNextPoseSearch/Source/AnimNextPoseSearch/Private/PoseHistoryChooserParameter.cpp b/Engine/Plugins/Experimental/AnimNextPoseSearch/Source/AnimNextPoseSearch/Private/PoseHistoryChooserParameter.cpp
--- a/Engine/Plugins/Experimental/AnimNextPoseSearch/Source/AnimNextPoseSearch/Private/PoseHistoryChooserParameter.cpp
+++ b/Engine/Plugins/Experimental/AnimNextPoseSearch/Source/AnimNextPoseSearch/Private/PoseHistoryChooserParameter.cpp
@@ -4,11 +4,11 @@
 #include "DataInterface/AnimNextDataInterfaceInstance.h"
 #include "StructUtils/PropertyBag.h"
 
-namespace UE::AnimNext::Private
+namespace 
 {
 
 // Helper function to find the first data interface instance in the context 
-FAnimNextDataInterfaceInstance* GetFirstDataInterfaceInstance(FChooserEvaluationContext& Context)
+FAnimNextDataInterfaceInstance* GetFirstAnimNextDataInterfaceInstance(FChooserEvaluationContext& Context)
 {
 	for(const FStructView& Param : Context.Params)
 	{
@@ -25,7 +25,7 @@ FAnimNextDataInterfaceInstance* GetFirstDataInterfaceInstance(FChooserEvaluation
 
 bool FPoseHistoryAnimProperty::GetValue(FChooserEvaluationContext& Context, FPoseHistoryReference& OutResult) const
 {
-	if(FAnimNextDataInterfaceInstance* Instance = UE::AnimNext::Private::GetFirstDataInterfaceInstance(Context))
+	if(FAnimNextDataInterfaceInstance* Instance = GetFirstAnimNextDataInterfaceInstance(Context))
 	{
 		return Instance->GetVariable(VariableName, OutResult) == EPropertyBagResult::Success;
 	}
