Add MultiAnim GetAnimationAsset RigVM node so Pose Search Interactions can be used in AN

diff --git a/Engine/Plugins/Experimental/AnimNextPoseSearch/Source/AnimNextPoseSearch/Private/RigUnit_MultiAnimGetAnimationAsset.cpp b/Engine/Plugins/Experimental/AnimNextPoseSearch/Source/AnimNextPoseSearch/Private/RigUnit_MultiAnimGetAnimationAsset.cpp
new file mode 100644
--- /dev/null
+++ b/Engine/Plugins/Experimental/AnimNextPoseSearch/Source/AnimNextPoseSearch/Private/RigUnit_MultiAnimGetAnimationAsset.cpp
@@ -0,0 +1,15 @@
+// Copyright Epic Games, Inc. All Rights Reserved.
+
+#include "RigUnit_MultiAnimGetAnimationAsset.h"
+
+#include UE_INLINE_GENERATED_CPP_BY_NAME(RigUnit_MultiAnimGetAnimationAsset)
+
+FRigUnit_MultiAnimGetAnimationAsset_Execute()
+{
+	const_cast<UObject*>(ExecuteContext.GetOwningObject());
+
+	if (MultiAnimAsset)
+	{
+		Result = MultiAnimAsset->GetAnimationAsset(Role);
+	}
+}
diff --git a/Engine/Plugins/Experimental/AnimNextPoseSearch/Source/AnimNextPoseSearch/Private/RigUnit_MultiAnimGetAnimationAsset.h b/Engine/Plugins/Experimental/AnimNextPoseSearch/Source/AnimNextPoseSearch/Private/RigUnit_MultiAnimGetAnimationAsset.h
new file mode 100644
--- /dev/null
+++ b/Engine/Plugins/Experimental/AnimNextPoseSearch/Source/AnimNextPoseSearch/Private/RigUnit_MultiAnimGetAnimationAsset.h
@@ -0,0 +1,31 @@
+// Copyright Epic Games, Inc. All Rights Reserved.
+
+#pragma once
+
+#include "Graph/RigUnit_AnimNextBase.h"
+#include "Animation/AnimationAsset.h"
+#include "PoseSearch/MultiAnimAsset.h"
+
+#include "RigUnit_MultiAnimGetAnimationAsset.generated.h"
+
+/** Gets anim asset for a particular role from a MultiAnimAsset */
+USTRUCT(meta=(DisplayName="Multi Anim Get Animation Asset", Category="Animation Graph", NodeColor="0, 1, 1", Keywords="MultiAnim"))
+struct FRigUnit_MultiAnimGetAnimationAsset : public FRigUnit_AnimNextBase
+{
+	GENERATED_BODY()
+
+	RIGVM_METHOD()
+	void Execute();
+
+	// MultiAnim Asset to search for role on
+	UPROPERTY(EditAnywhere, Category = "Graph", meta = (Input))
+	TObjectPtr<UMultiAnimAsset> MultiAnimAsset = nullptr;
+
+	// Role to search for within the MultiAnimAsset
+	UPROPERTY(EditAnywhere, Category = "Graph", meta = (Input))
+	FName Role = NAME_None;
+
+	// Result of the role search in the MultiAnimAsset
+	UPROPERTY(EditAnywhere, Category = "Graph", meta = (Output))
+	TObjectPtr<UAnimationAsset> Result = nullptr;
+};
