Fix unintended debug code submit RigUnit_MultiAnimGetAnimationAsset

diff --git a/Engine/Plugins/Experimental/AnimNextPoseSearch/Source/AnimNextPoseSearch/Private/RigUnit_MultiAnimGetAnimationAsset.cpp b/Engine/Plugins/Experimental/AnimNextPoseSearch/Source/AnimNextPoseSearch/Private/RigUnit_MultiAnimGetAnimationAsset.cpp
--- a/Engine/Plugins/Experimental/AnimNextPoseSearch/Source/AnimNextPoseSearch/Private/RigUnit_MultiAnimGetAnimationAsset.cpp
+++ b/Engine/Plugins/Experimental/AnimNextPoseSearch/Source/AnimNextPoseSearch/Private/RigUnit_MultiAnimGetAnimationAsset.cpp
@@ -6,8 +6,6 @@
 
 FRigUnit_MultiAnimGetAnimationAsset_Execute()
 {
-	const_cast<UObject*>(ExecuteContext.GetOwningObject());
-
 	if (MultiAnimAsset)
 	{
 		Result = MultiAnimAsset->GetAnimationAsset(Role);
