PoseSearch - MM blend spaces support improvements (more precise conversion between real to normalized time in blend spaces)

diff --git a/Engine/Plugins/Animation/PoseSearch/Source/Editor/Private/PoseSearchDebuggerDatabaseView.cpp b/Engine/Plugins/Animation/PoseSearch/Source/Editor/Private/PoseSearchDebuggerDatabaseView.cpp
--- a/Engine/Plugins/Animation/PoseSearch/Source/Editor/Private/PoseSearchDebuggerDatabaseView.cpp
+++ b/Engine/Plugins/Animation/PoseSearch/Source/Editor/Private/PoseSearchDebuggerDatabaseView.cpp
@@ -2,20 +2,21 @@
 
 #include "PoseSearchDebuggerDatabaseView.h"
 #include "Algo/AllOf.h"
+#include "Animation/AnimationSettings.h"
 #include "Animation/AnimComposite.h"
 #include "Animation/AnimMontage.h"
 #include "Animation/AnimSequence.h"
-#include "StructUtils/InstancedStruct.h"
 #include "Internationalization/Regex.h"
-#include "PoseSearchDebuggerDatabaseRow.h"
-#include "PoseSearchDebuggerSettings.h"
-#include "PoseSearchDebuggerView.h"
-#include "PoseSearchDebuggerViewModel.h"
 #include "PoseSearch/PoseSearchDatabase.h"
 #include "PoseSearch/PoseSearchDerivedData.h"
 #include "PoseSearch/PoseSearchFeatureChannel.h"
 #include "PoseSearch/PoseSearchSchema.h"
+#include "PoseSearchDebuggerDatabaseRow.h"
+#include "PoseSearchDebuggerSettings.h"
+#include "PoseSearchDebuggerView.h"
+#include "PoseSearchDebuggerViewModel.h"
 #include "PoseSearchEditor.h"
+#include "StructUtils/InstancedStruct.h"
 #include "Trace/PoseSearchTraceProvider.h"
 #include "Widgets/Input/SCheckBox.h"
 #include "Widgets/Layout/SScrollBar.h"
@@ -222,14 +223,17 @@ static void AddUnfilteredDatabaseRow(const UPoseSearchDatabase* Database, TConst
 		{
 			const FVector BlendParameters = SearchIndexAsset->GetBlendParameters();
 			const float PlayLength = DatabaseAsset->GetPlayLength(BlendParameters);
-			UObject* AnimationAsset = DatabaseAsset->GetAnimationAsset();
+			const UObject* AnimationAsset = DatabaseAsset->GetAnimationAsset();
 
 			Row->AssetName = DatabaseAsset->GetName();
 			Row->AssetPath = AnimationAsset ? AnimationAsset->GetPathName() : "";
 			Row->bLooping = DatabaseAsset->IsLooping();
 			Row->BlendParameters = BlendParameters;
-			Row->AnimFrame = DatabaseAsset->GetFrameAtTime(Time);
-			Row->AnimPercentage = FMath::IsNearlyZero(PlayLength) ? 0.f : Time / PlayLength;
+			
+			const float RealAssetTime = Database->GetRealAssetTime(DbPoseIdx);
+			const FFrameRate& DefaultFrameRate = UAnimationSettings::Get()->GetDefaultFrameRate();
+			Row->AnimFrame = DefaultFrameRate.AsFrameTime(RealAssetTime).RoundToFrame().Value;
+			Row->AnimPercentage = FMath::IsNearlyZero(PlayLength) ? 0.f : RealAssetTime / PlayLength;
 		}
 	}
 }
diff --git a/Engine/Plugins/Animation/PoseSearch/Source/Runtime/Private/AnimNode_MotionMatching.cpp b/Engine/Plugins/Animation/PoseSearch/Source/Runtime/Private/AnimNode_MotionMatching.cpp
--- a/Engine/Plugins/Animation/PoseSearch/Source/Runtime/Private/AnimNode_MotionMatching.cpp
+++ b/Engine/Plugins/Animation/PoseSearch/Source/Runtime/Private/AnimNode_MotionMatching.cpp
@@ -140,7 +140,7 @@ void FAnimNode_MotionMatching::UpdateAssetPlayer(const FAnimationUpdateContext&
 		// dynamic playback rate adjustment and as such the motion matching state does not update by itself
 		
 		// @todo: this is missing eventual update of the blend parameters that could change the pose!
-		MotionMatchingState.AdjustAssetTime(GetAccumulatedTime());
+		MotionMatchingState.CurrentSearchResult.UpdateWithNormalizedTime(GetAccumulatedTime());
 	}
 
 	UpdateCounter.SynchronizeWith(Context.AnimInstanceProxy->GetUpdateCounter());
diff --git a/Engine/Plugins/Animation/PoseSearch/Source/Runtime/Private/PoseSearchAssetSampler.cpp b/Engine/Plugins/Animation/PoseSearch/Source/Runtime/Private/PoseSearchAssetSampler.cpp
--- a/Engine/Plugins/Animation/PoseSearch/Source/Runtime/Private/PoseSearchAssetSampler.cpp
+++ b/Engine/Plugins/Animation/PoseSearch/Source/Runtime/Private/PoseSearchAssetSampler.cpp
@@ -281,10 +281,8 @@ float FAnimationAssetSampler::GetPlayLength(const UAnimationAsset* AnimAsset, co
 		{
 			TArray<FBlendSampleData> BlendSamples;
 			int32 TriangulationIndex = 0;
-			if (BlendSpace->GetSamplesFromBlendInput(BlendParameters, BlendSamples, TriangulationIndex, true))
-			{
-				PlayLength = BlendSpace->GetAnimationLengthFromSampleData(BlendSamples);
-			}
+			BlendSpace->GetSamplesFromBlendInput(BlendParameters, BlendSamples, TriangulationIndex, true);
+			PlayLength = BlendSpace->GetAnimationLengthFromSampleData(BlendSamples);
 		}
 		else
 		{
diff --git a/Engine/Plugins/Animation/PoseSearch/Source/Runtime/Private/PoseSearchDatabase.cpp b/Engine/Plugins/Animation/PoseSearch/Source/Runtime/Private/PoseSearchDatabase.cpp
--- a/Engine/Plugins/Animation/PoseSearch/Source/Runtime/Private/PoseSearchDatabase.cpp
+++ b/Engine/Plugins/Animation/PoseSearch/Source/Runtime/Private/PoseSearchDatabase.cpp
@@ -1,6 +1,11 @@
 // Copyright Epic Games, Inc. All Rights Reserved.
 
 #include "PoseSearch/PoseSearchDatabase.h"
+
+// @todo: remove this include once the deprecated method FPoseSearchDatabaseAnimationAssetBase::GetFrameAtTime is removed. 
+//        used only to get UAnimationSettings::Get()->GetDefaultFrameRate()
+#include "Animation/AnimationSettings.h"
+
 #include "Animation/AnimComposite.h"
 #include "Animation/AnimMontage.h"
 #include "Animation/AnimSequence.h"
@@ -294,20 +299,6 @@ static bool IsBlendSpaceRootMotionEnabled(const UBlendSpace* BlendSpace)
 	return bIsRootMotionUsedInBlendSpace;
 }
 
-static bool IsDatabaseAssetIndexBlendSpace(const UPoseSearchDatabase& Database, int32 AnimationAssetIndex)
-{
-	if (Database.GetDatabaseAnimationAsset<FPoseSearchDatabaseBlendSpace>(AnimationAssetIndex))
-	{
-		return true;
-	}
-	else if (const FPoseSearchDatabaseMultiAnimAsset* MultiAnimAsset =  Database.GetDatabaseAnimationAsset<FPoseSearchDatabaseMultiAnimAsset>(AnimationAssetIndex))
-	{
-		return MultiAnimAsset->GetAnimationAssetForRole(Database.Schema->GetDefaultRole())->IsA<UBlendSpace>();
-	}
-
-	return false;
-}
-
 } // namespace UE::PoseSearch
 
 //////////////////////////////////////////////////////////////////////////
@@ -328,7 +319,11 @@ int32 FPoseSearchDatabaseAnimationAssetBase::GetFrameAtTime(float Time) const
 	{
 		return SequenceBase->GetFrameAtTime(Time);
 	}
-	return 0.f;
+
+	// estimating the frame for blend spaces in a non precise way. Anyways this method is deprecated and should not be used
+	const float RealAssetTime = GetPlayLength(FVector::ZeroVector);
+	const FFrameRate& DefaultFrameRate = UAnimationSettings::Get()->GetDefaultFrameRate();
+	return DefaultFrameRate.AsFrameTime(RealAssetTime).RoundToFrame().Value;
 }
 
 bool FPoseSearchDatabaseAnimationAssetBase::IsSkeletonCompatible(TObjectPtr<const UPoseSearchSchema> InSchema) const
@@ -477,13 +472,22 @@ UObject* FPoseSearchDatabaseBlendSpace::GetAnimationAsset() const
 	return BlendSpace.Get();
 }
 
-#if WITH_EDITOR
-int32 FPoseSearchDatabaseBlendSpace::GetFrameAtTime(float Time) const
+float FPoseSearchDatabaseBlendSpace::GetPlayLength(const FVector& BlendParameters) const
 {
-	// returning the percentage of time as value to diplay in the pose search debugger (NoTe: BlendSpace->GetPlayLength() is one)
-	return FMath::RoundToInt(Time * 100.f);
+	int32 TriangulationIndex = 0;
+	TArray<FBlendSampleData> BlendSamples;
+	BlendSpace->GetSamplesFromBlendInput(BlendParameters, BlendSamples, TriangulationIndex, true);
+	const float PlayLength = BlendSpace->GetAnimationLengthFromSampleData(BlendSamples);
+	return PlayLength;
 }
-#endif // WITH_EDITOR
+
+//#if WITH_EDITOR
+//int32 FPoseSearchDatabaseBlendSpace::GetFrameAtTime(float Time) const
+//{
+//	// returning the percentage of time as value to diplay in the pose search debugger (NoTe: BlendSpace->GetPlayLength() is one)
+//	return FMath::RoundToInt(Time * 100.f);
+//}
+//#endif // WITH_EDITOR
 
 #if WITH_EDITORONLY_DATA
 UClass* FPoseSearchDatabaseBlendSpace::GetAnimationAssetStaticClass() const
@@ -509,10 +513,7 @@ FFloatInterval FPoseSearchDatabaseBlendSpace::GetEffectiveSamplingRange(const FV
 {
 	if (BlendSpace)
 	{
-		int32 TriangulationIndex = 0;
-		TArray<FBlendSampleData> BlendSamples;
-		BlendSpace->GetSamplesFromBlendInput(BlendParameters, BlendSamples, TriangulationIndex, true);
-		const float PlayLength = BlendSpace->GetAnimationLengthFromSampleData(BlendSamples);
+		const float PlayLength = GetPlayLength(BlendParameters);
 
 		// scaling blend space SamplingRange from the space [0, 1] to [0, PlayLength] with PlayLength calculated from the BlendSamples
 		FFloatInterval ScaledSamplingRange = SamplingRange;
@@ -678,13 +679,6 @@ float FPoseSearchDatabaseMultiAnimAsset::GetPlayLength(const FVector& BlendParam
 	return MultiAnimAsset ? MultiAnimAsset->GetPlayLength(BlendParameters) : 0.f;
 }
 
-#if WITH_EDITOR
-int32 FPoseSearchDatabaseMultiAnimAsset::GetFrameAtTime(float Time) const
-{
-	return MultiAnimAsset ? MultiAnimAsset->GetFrameAtTime(Time) : 0;
-}
-#endif // WITH_EDITOR
-
 int32 FPoseSearchDatabaseMultiAnimAsset::GetNumRoles() const
 {
 	return MultiAnimAsset ? MultiAnimAsset->GetNumRoles() : 0;
@@ -1299,21 +1293,11 @@ float UPoseSearchDatabase::GetNormalizedAssetTime(int32 PoseIdx) const
 {
 	check(Schema);
 	const UE::PoseSearch::FSearchIndexAsset& Asset = GetSearchIndex().GetAssetForPose(PoseIdx);
-	const bool bIsBlendSpace = UE::PoseSearch::IsDatabaseAssetIndexBlendSpace(*this, Asset.GetSourceAssetIdx());
-
-	// sequences or anim composites
-	float AssetTime = Asset.GetTimeFromPoseIndex(PoseIdx, Schema->SampleRate);
-
-	if (bIsBlendSpace && Asset.GetNumPoses() > 1)
-	{
-		// For BlendSpaces the AssetTime is in the range [0, 1] while the Sampling Range
-		// is in real time (seconds). We should be using but FAnimationAssetSampler::GetPlayLength(...) to normalize precisely,
-		// but Asset.GetNumPoses() - 1 is a good enough estimator. FMath::Min(1, ...) is there to clamp numerical errors
-		AssetTime = FMath::Min(1.f, AssetTime * Schema->SampleRate / float(Asset.GetNumPoses() - 1));
-		check(AssetTime >= 0.f);
-	}
-
-	return AssetTime;
+	const float ToRealTimeFactor = Asset.GetToRealTimeFactor();
+	check(ToRealTimeFactor > UE_KINDA_SMALL_NUMBER);
+	const float RealAssetTime = Asset.GetTimeFromPoseIndex(PoseIdx, Schema->SampleRate);
+	const float NormalizedAssetTime = RealAssetTime / ToRealTimeFactor;
+	return NormalizedAssetTime;
 }
 
 UE::PoseSearch::FSearchResult UPoseSearchDatabase::Search(UE::PoseSearch::FSearchContext& SearchContext) const
diff --git a/Engine/Plugins/Animation/PoseSearch/Source/Runtime/Private/PoseSearchDerivedData.cpp b/Engine/Plugins/Animation/PoseSearch/Source/Runtime/Private/PoseSearchDerivedData.cpp
--- a/Engine/Plugins/Animation/PoseSearch/Source/Runtime/Private/PoseSearchDerivedData.cpp
+++ b/Engine/Plugins/Animation/PoseSearch/Source/Runtime/Private/PoseSearchDerivedData.cpp
@@ -430,9 +430,17 @@ static bool InitSearchIndexAssets(FSearchIndexBase& SearchIndex, const UPoseSear
 			const bool bAddMirrored = DatabaseAsset->GetMirrorOption() == EPoseSearchMirrorOption::MirroredOnly || DatabaseAsset->GetMirrorOption() == EPoseSearchMirrorOption::UnmirroredAndMirrored;
 			const bool bIsLooping = DatabaseAsset->IsLooping();
 			const bool bDisableReselection = DatabaseAsset->IsDisableReselection();
+			
+			// @todo: add better support for IMultiAnimAsset: currently we fix blend space parameters for only one role, 
+			//        that implies having homogeneous blendspaces for ALL the roles, and having either ONLY blend spaces or only NOT blendspaces..
+			bool bIsBlendSpace = false;
+			if (UAnimationAsset* DefaultRoleAnimationAsset = DatabaseAsset->GetAnimationAssetForRole(DatabaseToLookForAssets->Schema->GetDefaultRole()))
+			{
+				bIsBlendSpace = DefaultRoleAnimationAsset->IsA<UBlendSpace>();
+			}
 
 			DatabaseAsset->IterateOverSamplingParameter(
-				[DatabaseAsset, ExcludeFromDatabaseParameters, &ValidRanges, Schema, bAddUnmirrored, bAddMirrored, bIsLooping, bDisableReselection, AnimationAssetIndex, &TotalPoses, &SearchIndex]
+				[DatabaseAsset, ExcludeFromDatabaseParameters, &ValidRanges, Schema, bAddUnmirrored, bAddMirrored, bIsLooping, bIsBlendSpace, bDisableReselection, AnimationAssetIndex, &TotalPoses, &SearchIndex]
 				(const FVector& BlendParameters)
 				{
 					FindValidRanges(DatabaseAsset, BlendParameters, ExcludeFromDatabaseParameters, ValidRanges);
@@ -443,10 +451,20 @@ static bool InitSearchIndexAssets(FSearchIndexBase& SearchIndex, const UPoseSear
 						{
 							const FFloatInterval RangeInterval(Range.GetLowerBoundValue(), Range.GetUpperBoundValue());
 
+							float ToRealTimeFactor = 1.f;
+							if (bIsBlendSpace)
+							{
+								const float PlayLength = DatabaseAsset->GetPlayLength(BlendParameters);
+								if (PlayLength > UE_KINDA_SMALL_NUMBER)
+								{
+									ToRealTimeFactor = PlayLength;
+								}
+							}
+
 							if (bAddUnmirrored)
 							{
 								const FSearchIndexAsset PoseSearchIndexAsset(AnimationAssetIndex, TotalPoses, false, bIsLooping,
-									bDisableReselection, RangeInterval, Schema->SampleRate, PermutationIdx, BlendParameters);
+									bDisableReselection, RangeInterval, Schema->SampleRate, PermutationIdx, BlendParameters, ToRealTimeFactor);
 								if (PoseSearchIndexAsset.GetNumPoses() > 0)
 								{
 									SearchIndex.Assets.Add(PoseSearchIndexAsset);
@@ -457,7 +475,7 @@ static bool InitSearchIndexAssets(FSearchIndexBase& SearchIndex, const UPoseSear
 							if (bAddMirrored)
 							{
 								const FSearchIndexAsset PoseSearchIndexAsset(AnimationAssetIndex, TotalPoses, true, bIsLooping,
-									bDisableReselection, RangeInterval, Schema->SampleRate, PermutationIdx, BlendParameters);
+									bDisableReselection, RangeInterval, Schema->SampleRate, PermutationIdx, BlendParameters, ToRealTimeFactor);
 								if (PoseSearchIndexAsset.GetNumPoses() > 0)
 								{
 									SearchIndex.Assets.Add(PoseSearchIndexAsset);
diff --git a/Engine/Plugins/Animation/PoseSearch/Source/Runtime/Private/PoseSearchHistory.cpp b/Engine/Plugins/Animation/PoseSearch/Source/Runtime/Private/PoseSearchHistory.cpp
--- a/Engine/Plugins/Animation/PoseSearch/Source/Runtime/Private/PoseSearchHistory.cpp
+++ b/Engine/Plugins/Animation/PoseSearch/Source/Runtime/Private/PoseSearchHistory.cpp
@@ -789,7 +789,7 @@ bool FPoseHistory::GetTransformAtTime(float Time, FTransform& OutBoneTransform,
 		}
 	}
 
-	// @todo: re-enable this logging for editor only builds once the original issue has been fixed
+	// @todo: reenable this logging after implementing AnimNext PoseHistory initialization (currently spamming on actor spawning with MM active)
 //#if WITH_EDITOR
 //	// @todo: logging errors only if pose history was properly initialized for now (since AnimNext version doesn't implement initialization with ref pose yet),
 //	// and we ultimately want to catch errors of missing bone transforms only checking valid USkeleton(s) against valid PoseSearchSchema(s)
diff --git a/Engine/Plugins/Animation/PoseSearch/Source/Runtime/Private/PoseSearchIndex.cpp b/Engine/Plugins/Animation/PoseSearch/Source/Runtime/Private/PoseSearchIndex.cpp
--- a/Engine/Plugins/Animation/PoseSearch/Source/Runtime/Private/PoseSearchIndex.cpp
+++ b/Engine/Plugins/Animation/PoseSearch/Source/Runtime/Private/PoseSearchIndex.cpp
@@ -246,7 +246,8 @@ bool FSearchIndexAsset::operator==(const FSearchIndexAsset& Other) const
 		BlendParameterY == Other.BlendParameterY &&
 		FirstPoseIdx == Other.FirstPoseIdx &&
 		FirstSampleIdx == Other.FirstSampleIdx &&
-		LastSampleIdx == Other.LastSampleIdx;
+		LastSampleIdx == Other.LastSampleIdx &&
+		ToRealTimeFactor == Other.ToRealTimeFactor;
 }
 
 FArchive& operator<<(FArchive& Ar, FSearchIndexAsset& IndexAsset)
@@ -261,6 +262,7 @@ FArchive& operator<<(FArchive& Ar, FSearchIndexAsset& IndexAsset)
 	int32 FirstPoseIdx = IndexAsset.FirstPoseIdx;
 	int32 FirstSampleIdx = IndexAsset.FirstSampleIdx;
 	int32 LastSampleIdx = IndexAsset.LastSampleIdx;
+	float ToRealTimeFactor = IndexAsset.ToRealTimeFactor;
 
 	// @todo: optimize the archived size of FSearchIndexAsset, since most members are bitfields
 	Ar << SourceAssetIdx;
@@ -273,9 +275,10 @@ FArchive& operator<<(FArchive& Ar, FSearchIndexAsset& IndexAsset)
 	Ar << FirstPoseIdx;
 	Ar << FirstSampleIdx;
 	Ar << LastSampleIdx;
+	Ar << ToRealTimeFactor;
 
 	new(&IndexAsset) FSearchIndexAsset(SourceAssetIdx, bMirrored, bLooping, bDisableReselection,
-		PermutationIdx, FVector(BlendParameterX, BlendParameterY, 0.f), FirstPoseIdx, FirstSampleIdx, LastSampleIdx);
+		PermutationIdx, FVector(BlendParameterX, BlendParameterY, 0.f), FirstPoseIdx, FirstSampleIdx, LastSampleIdx, ToRealTimeFactor);
 
 	return Ar;
 }
diff --git a/Engine/Plugins/Animation/PoseSearch/Source/Runtime/Private/PoseSearchInteractionAsset.cpp b/Engine/Plugins/Animation/PoseSearch/Source/Runtime/Private/PoseSearchInteractionAsset.cpp
--- a/Engine/Plugins/Animation/PoseSearch/Source/Runtime/Private/PoseSearchInteractionAsset.cpp
+++ b/Engine/Plugins/Animation/PoseSearch/Source/Runtime/Private/PoseSearchInteractionAsset.cpp
@@ -103,45 +103,6 @@ float UPoseSearchInteractionAsset::GetPlayLength(const FVector& BlendParameters)
 	return MaxPlayLength;
 }
 
-#if WITH_EDITOR
-// @todo: pass const FVector& BlendParameters as input and use a similar logic to GetPlayLength to calculate BS PlayLength
-int32 UPoseSearchInteractionAsset::GetFrameAtTime(float Time) const
-{
-	const UAnimationAsset* MaxPlayLengthAnim = nullptr;
-	float MaxPlayLength = -1.f;
-	for (const FPoseSearchInteractionAssetItem& Item : Items)
-	{
-		if (const UAnimationAsset* AnimationAsset = Item.Animation.Get())
-		{
-			const float PlayLength = AnimationAsset->GetPlayLength();
-			if (PlayLength > MaxPlayLength)
-			{
-				MaxPlayLength = PlayLength;
-				MaxPlayLengthAnim = AnimationAsset;
-			}
-		}
-	}
-	
-	if (MaxPlayLengthAnim)
-	{
-		if (const UAnimSequenceBase* SequenceBase = Cast<UAnimSequenceBase>(MaxPlayLengthAnim))
-		{
-			return SequenceBase->GetFrameAtTime(Time);
-		}
-
-		if (const UBlendSpace* BlendSpace = Cast<UBlendSpace>(MaxPlayLengthAnim))
-		{
-			// returning the percentage of time as value to diplay in the pose search debugger (NoTe: BlendSpace->GetPlayLength() is one)
-			return FMath::RoundToInt(Time * 100.f);
-		}
-
-		unimplemented();
-	}
-
-	return 0;
-}
-#endif // WITH_EDITOR
-
 FQuat UPoseSearchInteractionAsset::FindReferenceOrientationNoBanking(const TArrayView<const FTransform> Transforms, const TArrayView<int32> SortedByWarpingWeightRotationItemIndex, const TArrayView<float> NormalizedWarpingWeightRotation) const
 {
 	check(Items.Num() > 0);
diff --git a/Engine/Plugins/Animation/PoseSearch/Source/Runtime/Private/PoseSearchInteractionSubsystem.cpp b/Engine/Plugins/Animation/PoseSearch/Source/Runtime/Private/PoseSearchInteractionSubsystem.cpp
--- a/Engine/Plugins/Animation/PoseSearch/Source/Runtime/Private/PoseSearchInteractionSubsystem.cpp
+++ b/Engine/Plugins/Animation/PoseSearch/Source/Runtime/Private/PoseSearchInteractionSubsystem.cpp
@@ -410,6 +410,8 @@ void UPoseSearchInteractionSubsystem::PopulateContinuingProperties(float DeltaSe
 						if (const FPoseSearchDatabaseAnimationAssetBase* DatabaseAsset = SearchResult->Database->GetDatabaseAnimationAsset<FPoseSearchDatabaseAnimationAssetBase>(*SearchIndexAsset))
 						{
 							SearchContext.PlayingAsset = DatabaseAsset->GetAnimationAsset();
+
+							// @todo: for blendspaces PlayingAssetAccumulatedTime should be incremented by a normalized DeltaSeconds
 							SearchContext.PlayingAssetAccumulatedTime = SearchResult->AssetTime + DeltaSeconds;
 							// @todo: populate SearchContext.InterruptMode
 						}
diff --git a/Engine/Plugins/Animation/PoseSearch/Source/Runtime/Private/PoseSearchLibrary.cpp b/Engine/Plugins/Animation/PoseSearch/Source/Runtime/Private/PoseSearchLibrary.cpp
--- a/Engine/Plugins/Animation/PoseSearch/Source/Runtime/Private/PoseSearchLibrary.cpp
+++ b/Engine/Plugins/Animation/PoseSearch/Source/Runtime/Private/PoseSearchLibrary.cpp
@@ -320,7 +320,7 @@ PRAGMA_ENABLE_DEPRECATION_WARNINGS
 
 void FMotionMatchingState::AdjustAssetTime(float AssetTime)
 {
-	CurrentSearchResult.Update(AssetTime);
+	CurrentSearchResult.UpdateWithNormalizedTime(AssetTime);
 }
 
 FVector FMotionMatchingState::GetEstimatedFutureRootMotionVelocity() const
@@ -771,6 +771,7 @@ void UPoseSearchLibrary::UpdateMotionMatchingState(
 	}
 	else
 	{
+		// @todo: for blendspaces ElapsedPoseSearchTime should be incremented by a normalized DeltaTime, also didn't we already synchronized the result???
 		InOutMotionMatchingState.ElapsedPoseSearchTime += DeltaTime;
 		InOutMotionMatchingState.CurrentSearchResult.bIsContinuingPoseSearch = true;
 	}
diff --git a/Engine/Plugins/Animation/PoseSearch/Source/Runtime/Private/PoseSearchResult.cpp b/Engine/Plugins/Animation/PoseSearch/Source/Runtime/Private/PoseSearchResult.cpp
--- a/Engine/Plugins/Animation/PoseSearch/Source/Runtime/Private/PoseSearchResult.cpp
+++ b/Engine/Plugins/Animation/PoseSearch/Source/Runtime/Private/PoseSearchResult.cpp
@@ -9,29 +9,81 @@
 namespace UE::PoseSearch
 {
 
-void FSearchResult::Update(float NewAssetTime)
+void FSearchResult::UpdateWithNormalizedTime(float NormalizedTime)
 {
 	if (IsValid())
 	{
-		float RealTime = NewAssetTime;
+		// for non blend spaces the real time corrisponds to the normalized time!
 		const FSearchIndexAsset& SearchIndexAsset = Database->GetSearchIndex().GetAssetForPose(PoseIdx);
+		const float RealTime = NormalizedTime * SearchIndexAsset.GetToRealTimeFactor();
+		
+#if WITH_EDITOR && DO_CHECK
 		const FPoseSearchDatabaseAnimationAssetBase* DatabaseAnimationAssetBase = Database->GetDatabaseAnimationAsset<FPoseSearchDatabaseAnimationAssetBase>(SearchIndexAsset);
 		check(DatabaseAnimationAssetBase);
-		if (const UBlendSpace* BlendSpace = Cast<UBlendSpace>(DatabaseAnimationAssetBase->GetAnimationAsset()))
+		if (Cast<UBlendSpace>(DatabaseAnimationAssetBase->GetAnimationAssetForRole(Database->Schema->GetDefaultRole())))
+		{
+			const float PlayLength = DatabaseAnimationAssetBase->GetPlayLength(SearchIndexAsset.GetBlendParameters());
+			if (PlayLength > UE_KINDA_SMALL_NUMBER)
+			{
+				// Asset player time for blendspaces is normalized [0, 1] so we need to convert 
+				// to a real time before we advance it
+				check(NormalizedTime >= 0.f && NormalizedTime <= 1.f);
+				check(FMath::IsNearlyEqual(RealTime, NormalizedTime * PlayLength));
+			}
+			else
+			{
+				check(FMath::IsNearlyEqual(SearchIndexAsset.GetToRealTimeFactor(), 1.f));
+			}
+		}
+		else
 		{
-			TArray<FBlendSampleData> BlendSamples;
-			int32 TriangulationIndex = 0;
-			BlendSpace->GetSamplesFromBlendInput(SearchIndexAsset.GetBlendParameters(), BlendSamples, TriangulationIndex, true);
-			const float PlayLength = BlendSpace->GetAnimationLengthFromSampleData(BlendSamples);
-
-			// Asset player time for blendspaces is normalized [0, 1] so we need to convert 
-			// to a real time before we advance it
-			check(NewAssetTime >= 0.f && NewAssetTime <= 1.f);
-			RealTime = NewAssetTime * PlayLength;
+			check(FMath::IsNearlyEqual(SearchIndexAsset.GetToRealTimeFactor(), 1.f));
 		}
+#endif // WITH_EDITOR && DO_CHECK
 		
 		PoseIdx = Database->GetPoseIndexFromTime(RealTime, SearchIndexAsset);
-		AssetTime = NewAssetTime;
+		AssetTime = NormalizedTime;
+	}
+	else
+	{
+		Reset();
+	}
+}
+
+void FSearchResult::UpdateWithRealTime(float RealTime)
+{
+	if (IsValid())
+	{
+		const FSearchIndexAsset& SearchIndexAsset = Database->GetSearchIndex().GetAssetForPose(PoseIdx);
+		check(SearchIndexAsset.GetToRealTimeFactor() > UE_KINDA_SMALL_NUMBER);
+		const float NormalizedTime = RealTime / SearchIndexAsset.GetToRealTimeFactor();
+
+#if WITH_EDITOR && DO_CHECK
+		const FPoseSearchDatabaseAnimationAssetBase* DatabaseAnimationAssetBase = Database->GetDatabaseAnimationAsset<FPoseSearchDatabaseAnimationAssetBase>(SearchIndexAsset);
+		check(DatabaseAnimationAssetBase);
+		if (Cast<UBlendSpace>(DatabaseAnimationAssetBase->GetAnimationAssetForRole(Database->Schema->GetDefaultRole())))
+		{
+			const float PlayLength = DatabaseAnimationAssetBase->GetPlayLength(SearchIndexAsset.GetBlendParameters());
+			if (PlayLength > UE_KINDA_SMALL_NUMBER)
+			{
+				// Asset player time for blendspaces is normalized [0, 1] so we need to convert 
+				// to a real time before we advance it
+				check(RealTime >= 0.f && RealTime <= PlayLength);
+				check(FMath::IsNearlyEqual(NormalizedTime, PlayLength > UE_KINDA_SMALL_NUMBER ? RealTime / PlayLength : 0.f));
+			}
+			else
+			{
+				check(FMath::IsNearlyEqual(SearchIndexAsset.GetToRealTimeFactor(), 1.f));
+			}
+		}
+		else
+		{
+			check(FMath::IsNearlyEqual(SearchIndexAsset.GetToRealTimeFactor(), 1.f));
+		}
+#endif // WITH_EDITOR && DO_CHECK
+
+		PoseIdx = Database->GetPoseIndexFromTime(RealTime, SearchIndexAsset);
+		AssetTime = NormalizedTime;
 	}
 	else
 	{
diff --git a/Engine/Plugins/Animation/PoseSearch/Source/Runtime/Public/PoseSearch/MultiAnimAsset.h b/Engine/Plugins/Animation/PoseSearch/Source/Runtime/Public/PoseSearch/MultiAnimAsset.h
--- a/Engine/Plugins/Animation/PoseSearch/Source/Runtime/Public/PoseSearch/MultiAnimAsset.h
+++ b/Engine/Plugins/Animation/PoseSearch/Source/Runtime/Public/PoseSearch/MultiAnimAsset.h
@@ -22,7 +22,6 @@ public:
 	[[nodiscard]] virtual float GetPlayLength(const FVector& BlendParameters) const PURE_VIRTUAL(UMultiAnimAsset::GetPlayLength, return 0.f;);
 
 #if WITH_EDITOR
-	[[nodiscard]] virtual int32 GetFrameAtTime(float Time) const PURE_VIRTUAL(UMultiAnimAsset::GetFrameAtTime, return 0;);
 	[[nodiscard]] virtual USkeletalMesh* GetPreviewMesh(const FName& Role) const PURE_VIRTUAL(UMultiAnimAsset::GetPreviewMesh, return nullptr;);
 #endif // WITH_EDITOR
 
diff --git a/Engine/Plugins/Animation/PoseSearch/Source/Runtime/Public/PoseSearch/PoseSearchDatabase.h b/Engine/Plugins/Animation/PoseSearch/Source/Runtime/Public/PoseSearch/PoseSearchDatabase.h
--- a/Engine/Plugins/Animation/PoseSearch/Source/Runtime/Public/PoseSearch/PoseSearchDatabase.h
+++ b/Engine/Plugins/Animation/PoseSearch/Source/Runtime/Public/PoseSearch/PoseSearchDatabase.h
@@ -69,6 +69,7 @@ struct POSESEARCH_API FPoseSearchDatabaseAnimationAssetBase
 	virtual UAnimationAsset* GetAnimationAssetForRole(const UE::PoseSearch::FRole& Role) const;
 	virtual FTransform GetRootTransformOriginForRole(const UE::PoseSearch::FRole& Role) const;
 #if WITH_EDITOR
+	UE_DEPRECATED(5.6, "GetFrameAtTime is no longer supported")
 	virtual int32 GetFrameAtTime(float Time) const;
 	virtual bool IsSkeletonCompatible(TObjectPtr<const UPoseSearchSchema> InSchema) const;
 	// Experimental, this feature might be removed without warning, not for production use
@@ -240,10 +241,7 @@ struct POSESEARCH_API FPoseSearchDatabaseBlendSpace : public FPoseSearchDatabase
 #endif // WITH_EDITORONLY_DATA
 
 	virtual UObject* GetAnimationAsset() const override;
-
-#if WITH_EDITOR
-	virtual int32 GetFrameAtTime(float Time) const override;
-#endif // WITH_EDITOR
+	virtual float GetPlayLength(const FVector& BlendParameters) const override;
 
 	friend bool operator==(const FPoseSearchDatabaseBlendSpace& A, const FPoseSearchDatabaseBlendSpace& B)
 	{
@@ -398,10 +396,6 @@ struct POSESEARCH_API FPoseSearchDatabaseMultiAnimAsset : public FPoseSearchData
 	virtual UAnimationAsset* GetAnimationAssetForRole(const UE::PoseSearch::FRole& Role) const override;
 	virtual FTransform GetRootTransformOriginForRole(const UE::PoseSearch::FRole& Role) const override;
 
-#if WITH_EDITOR
-	virtual int32 GetFrameAtTime(float Time) const override;
-#endif // WITH_EDITOR
-
 	friend bool operator==(const FPoseSearchDatabaseMultiAnimAsset& A, const FPoseSearchDatabaseMultiAnimAsset& B)
 	{
 		return static_cast<const FPoseSearchDatabaseAnimationAssetBase&>(A) == static_cast<const FPoseSearchDatabaseAnimationAssetBase&>(B)
diff --git a/Engine/Plugins/Animation/PoseSearch/Source/Runtime/Public/PoseSearch/PoseSearchDefines.h b/Engine/Plugins/Animation/PoseSearch/Source/Runtime/Public/PoseSearch/PoseSearchDefines.h
--- a/Engine/Plugins/Animation/PoseSearch/Source/Runtime/Public/PoseSearch/PoseSearchDefines.h
+++ b/Engine/Plugins/Animation/PoseSearch/Source/Runtime/Public/PoseSearch/PoseSearchDefines.h
@@ -42,7 +42,7 @@ static constexpr int32 PreallocatedSearchesNum = 4;
 static constexpr float FiniteDelta = 1 / 60.0f; // Time delta used for computing pose derivatives
 static constexpr int32 MaxNumberOfCollectedPoseCandidatesPerDatabase = 256;
 static constexpr int32 InvalidSearchId = INDEX_NONE; // @todo: it should be 0, since the search id is an hash!
-static constexpr int32 DatabaseIndexDerivedDataCacheKeyVersion = 2683; // bump this value to force reindex ALL the UPoseSearchDatabases, without having to change the global FDevSystemGuids::POSESEARCHDB_DERIVEDDATA_VER
+static constexpr int32 DatabaseIndexDerivedDataCacheKeyVersion = 2687; // bump this value to force reindex ALL the UPoseSearchDatabases, without having to change the global FDevSystemGuids::POSESEARCHDB_DERIVEDDATA_VER
 static constexpr int32 DefaultTickPriority = 100;
 
 template<uint32 Alignment = DEFAULT_ALIGNMENT> class TMemStackSetAllocator : public TSetAllocator<TSparseArrayAllocator<TMemStackAllocator<Alignment>, TMemStackAllocator<Alignment>>, TMemStackAllocator<Alignment>> {};
diff --git a/Engine/Plugins/Animation/PoseSearch/Source/Runtime/Public/PoseSearch/PoseSearchIndex.h b/Engine/Plugins/Animation/PoseSearch/Source/Runtime/Public/PoseSearch/PoseSearchIndex.h
--- a/Engine/Plugins/Animation/PoseSearch/Source/Runtime/Public/PoseSearch/PoseSearchIndex.h
+++ b/Engine/Plugins/Animation/PoseSearch/Source/Runtime/Public/PoseSearch/PoseSearchIndex.h
@@ -96,7 +96,8 @@ struct FSearchIndexAsset
 		const FVector& InBlendParameters,
 		int32 InFirstPoseIdx,
 		int32 InFirstSampleIdx,
-		int32 InLastSampleIdx)
+		int32 InLastSampleIdx,
+		float InToRealTimeFactor /*= 1.f*/)
 		: SourceAssetIdx(InSourceAssetIdx)
 		, bMirrored(bInMirrored)
 		, bLooping(bInLooping)
@@ -107,8 +108,10 @@ struct FSearchIndexAsset
 		, FirstPoseIdx(InFirstPoseIdx)
 		, FirstSampleIdx(InFirstSampleIdx)
 		, LastSampleIdx(InLastSampleIdx)
+		, ToRealTimeFactor(InToRealTimeFactor)
 	{
 		check(FMath::IsNearlyZero(InBlendParameters.Z));
+		check(ToRealTimeFactor > UE_KINDA_SMALL_NUMBER);		
 	}
 
 	FSearchIndexAsset(
@@ -120,7 +123,8 @@ struct FSearchIndexAsset
 		const FFloatInterval& InSamplingInterval,
 		int32 SchemaSampleRate,
 		int32 InPermutationIdx,
-		FVector InBlendParameters = FVector::Zero())
+		FVector InBlendParameters /*= FVector::Zero()*/,
+		float InToRealTimeFactor /*= 1.f*/)
 		: FSearchIndexAsset(
 			InSourceAssetIdx,
 			bInMirrored,
@@ -130,7 +134,8 @@ struct FSearchIndexAsset
 			InBlendParameters,
 			InFirstPoseIdx,
 			FMath::CeilToInt(InSamplingInterval.Min * SchemaSampleRate),
-			FMath::FloorToInt(InSamplingInterval.Max * SchemaSampleRate))
+			FMath::FloorToInt(InSamplingInterval.Max * SchemaSampleRate),
+			InToRealTimeFactor)
 	{
 		check(SchemaSampleRate > 0);
 	}
@@ -142,6 +147,7 @@ struct FSearchIndexAsset
 	bool IsDisableReselection() const { return bDisableReselection; }
 	int32 GetPermutationIdx() const { return PermutationIdx; }
 	FVector GetBlendParameters() const { return FVector(BlendParameterX, BlendParameterY, 0.f); }
+	float GetToRealTimeFactor() const { return ToRealTimeFactor; }
 	int32 GetFirstPoseIdx() const { return FirstPoseIdx; }
 
 	bool IsPoseInRange(int32 PoseIdx) const { return (PoseIdx >= FirstPoseIdx) && (PoseIdx < FirstPoseIdx + GetNumPoses()); }
@@ -266,15 +272,22 @@ struct FSearchIndexAsset
 private:
 	// Index of the source asset in search index's container (i.e. UPoseSearchDatabase)
 	const int32 SourceAssetIdx = INDEX_NONE;
+
 	const bool bMirrored : 1 = false; 
 	const bool bLooping : 1 = false;
 	const bool bDisableReselection : 1 = false;
 	const int32 PermutationIdx = INDEX_NONE;
+	
+	// fixed blend space patameters associated to this FSearchIndexAsset if SourceAssetIdx refers to a blend space
 	const float BlendParameterX = 0.f;
 	const float BlendParameterY = 0.f;
+
 	const int32 FirstPoseIdx = INDEX_NONE;
 	const int32 FirstSampleIdx = INDEX_NONE;
 	const int32 LastSampleIdx = INDEX_NONE;
+
+	// if SourceAssetIdx refers to a blend space, the animation asset would be in a normalized time interval [0-1]. ToRealTimeFactor brings the normalized time to the real time for the fixed BlendParameterX/BlendParameterY
+	const float ToRealTimeFactor = 1.f;
 };
 
 struct FSearchStats
diff --git a/Engine/Plugins/Animation/PoseSearch/Source/Runtime/Public/PoseSearch/PoseSearchInteractionAsset.h b/Engine/Plugins/Animation/PoseSearch/Source/Runtime/Public/PoseSearch/PoseSearchInteractionAsset.h
--- a/Engine/Plugins/Animation/PoseSearch/Source/Runtime/Public/PoseSearch/PoseSearchInteractionAsset.h
+++ b/Engine/Plugins/Animation/PoseSearch/Source/Runtime/Public/PoseSearch/PoseSearchInteractionAsset.h
@@ -76,10 +76,6 @@ public:
 	virtual bool HasRootMotion() const override;
 	virtual float GetPlayLength(const FVector& BlendParameters) const override;
 
-#if WITH_EDITOR
-	virtual int32 GetFrameAtTime(float Time) const override;
-#endif // WITH_EDITOR
-
 	virtual int32 GetNumRoles() const override { return Items.Num(); }
 	virtual UE::PoseSearch::FRole GetRole(int32 RoleIndex) const override { return Items[RoleIndex].Role; }
 	
diff --git a/Engine/Plugins/Animation/PoseSearch/Source/Runtime/Public/PoseSearch/PoseSearchLibrary.h b/Engine/Plugins/Animation/PoseSearch/Source/Runtime/Public/PoseSearch/PoseSearchLibrary.h
--- a/Engine/Plugins/Animation/PoseSearch/Source/Runtime/Public/PoseSearch/PoseSearchLibrary.h
+++ b/Engine/Plugins/Animation/PoseSearch/Source/Runtime/Public/PoseSearch/PoseSearchLibrary.h
@@ -54,9 +54,7 @@ PRAGMA_ENABLE_DEPRECATION_WARNINGS
 	// Reset the state to a default state using the current Database
 	POSESEARCH_API void Reset();
 
-	// Attempts to set the internal state to match the provided asset time including updating the internal DbPoseIdx. 
-	// If the provided asset time is out of bounds for the currently playing asset then this function will reset the 
-	// state back to the default state.
+	UE_DEPRECATED(5.6, "Use CurrentSearchResult.UpdateWithNormalizedTime or CurrentSearchResult.UpdateWithRealTime instead")
 	POSESEARCH_API void AdjustAssetTime(float AssetTime);
 
 	POSESEARCH_API void UpdateWantedPlayRate(const UE::PoseSearch::FSearchContext& SearchContext, const FFloatInterval& PlayRate, float TrajectorySpeedMultiplier, const FPoseSearchEvent& EventToSearch /*= FPoseSearchEvent()*/);
diff --git a/Engine/Plugins/Animation/PoseSearch/Source/Runtime/Public/PoseSearch/PoseSearchResult.h b/Engine/Plugins/Animation/PoseSearch/Source/Runtime/Public/PoseSearch/PoseSearchResult.h
--- a/Engine/Plugins/Animation/PoseSearch/Source/Runtime/Public/PoseSearch/PoseSearchResult.h
+++ b/Engine/Plugins/Animation/PoseSearch/Source/Runtime/Public/PoseSearch/PoseSearchResult.h
@@ -31,7 +31,7 @@ struct FSearchResult
 	TWeakObjectPtr<const UPoseSearchDatabase> Database;
 
 	// time of the associated GetSearchIndexAsset from this search. During continuing pose searches,
-	// AssetTime is integrated by DeltaSeconds, so it's not precisely the quantized time associated to PoseIdx
+	// AssetTime is resynchronized via UpdateWithNormalizedTime or UpdateWithRealTime, so it's not precisely the quantized time associated to PoseIdx
 	float AssetTime = 0.f;
 
 	// true if the search result is from a continuing pose search
@@ -46,7 +46,8 @@ struct FSearchResult
 	// Attempts to set the internal state to match the provided asset time including updating the internal DbPoseIdx. 
 	// If the provided asset time is out of bounds for the currently playing asset then this function will reset the 
 	// state back to the default state.
-	void Update(float NewAssetTime);
+	POSESEARCH_API void UpdateWithNormalizedTime(float NormalizedTime);
+	POSESEARCH_API void UpdateWithRealTime(float RealTime);
 
 	bool IsValid() const
 	{
@@ -160,7 +161,7 @@ public:
 
 	// assigned role to this character. If this result is from a motion matched interaction, every interacting Actors will be assigned a different Role
 	UPROPERTY(Transient, VisibleAnywhere, BlueprintReadOnly, Category="Experimental|Interaction", meta = (Experimental))
-	FName Role;
+	FName Role = UE::PoseSearch::DefaultRole;
 
 	// cached actors root transforms for all the roles in SelectedAnim (as UMultiAnimAsset). For all the optional missing roles (where the associated AnimContexts[i] is null), the ActorRootTransform will be identity
 	UPROPERTY(Transient, VisibleAnywhere, BlueprintReadOnly, Category="Experimental|Interaction", meta = (Experimental))
diff --git a/Engine/Plugins/Experimental/AnimNextPoseSearch/Source/AnimNextPoseSearch/Private/MotionMatchingTrait.cpp b/Engine/Plugins/Experimental/AnimNextPoseSearch/Source/AnimNextPoseSearch/Private/MotionMatchingTrait.cpp
--- a/Engine/Plugins/Experimental/AnimNextPoseSearch/Source/AnimNextPoseSearch/Private/MotionMatchingTrait.cpp
+++ b/Engine/Plugins/Experimental/AnimNextPoseSearch/Source/AnimNextPoseSearch/Private/MotionMatchingTrait.cpp
@@ -46,6 +46,13 @@ static FAutoConsoleVariableRef CVarAnimNextMotionMatchingPlayRateEnabled(
 	TEXT("1: Enable all usages of PlayRate.\n")
 	TEXT("2: Enable PlayRate in PoseSeach only (Not used in actual playback).\n")
 );
+
+static bool GVarAnimNextMotionMatchingUseTimelineRealTime = false;
+static FAutoConsoleVariableRef CVarAnimNextMotionMatchingUseTimelineRealTime(
+	TEXT("a.AnimNext.MotionMatchingTrait.UseTimelineRealTime"),
+	GVarAnimNextMotionMatchingUseTimelineRealTime,
+	TEXT("Debug code to use the real time from the timeline rather the reconstructed normalized time for blend spaces to update the continuing pose")
+);
 #endif // ENABLE_ANIM_DEBUG
 
 namespace UE::AnimNext
@@ -202,7 +209,8 @@ namespace UE::AnimNext
 			if (MotionMatchingState.CurrentSearchResult.Database == nullptr || MotionMatchingState.CurrentSearchResult.Database->Schema == nullptr)
 			{
 			}
-	#if WITH_EDITOR
+
+#if WITH_EDITOR
 			else if (EAsyncBuildIndexResult::Success != FAsyncPoseSearchDatabasesManagement::RequestAsyncBuildIndex(MotionMatchingState.CurrentSearchResult.Database.Get(), ERequestAsyncBuildFlag::ContinueRequest))
 			{
 				// MotionMatchingState.CurrentSearchResult.Database is indexing, and it's not safe to use its previous index properties cached in MotionMatchingState
@@ -215,7 +223,7 @@ namespace UE::AnimNext
 				// CurrentSearchResult being out of synch with the updated database, so we need to reset the state
 				MotionMatchingState.Reset();
 			}
-	#endif // WITH_EDITOR
+#endif // WITH_EDITOR
 			// @todo: implement this MotionMatchingState.Reset() condition for parity to FAnimNode_MotionMatching::UpdateAssetPlayer
 			//else if (GetAnimAsset() != MotionMatchingState.CurrentSearchResult.GetCurrentResultAnimationAsset())
 			//{
@@ -229,27 +237,58 @@ namespace UE::AnimNext
 				// because the player node may have ticked more or less time than we expected due to variable dt or the 
 				// dynamic playback rate adjustment and as such the MotionMatchingState does not update by itself
 				const FTimelineState TimelineState = TimelineTrait.GetState(Context);
-				float TimelinePosition = TimelineState.GetPosition();
+				const float TimelineRealTime = TimelineState.GetPosition();
+				float TimelineNormalizedTime = -1.f;
 
-				// @TODO: This is a hack, since we don't have a way to get the normalized timeline of a blendspace yet (pending syncing...?)
-				if (const UE::PoseSearch::FSearchIndexAsset* SearchIndexAsset = MotionMatchingState.CurrentSearchResult.GetSearchIndexAsset())
+				bool bUseTimelineRealTime = false;
+#if ENABLE_ANIM_DEBUG
+				bUseTimelineRealTime = GVarAnimNextMotionMatchingUseTimelineRealTime;
+#endif // ENABLE_ANIM_DEBUG
+
+				if (bUseTimelineRealTime)
+				{
+					MotionMatchingState.CurrentSearchResult.UpdateWithRealTime(TimelineRealTime);
+				}
+				else
 				{
-					if (const FPoseSearchDatabaseAnimationAssetBase* DatabaseAnimationAssetBase = MotionMatchingState.CurrentSearchResult.Database->GetDatabaseAnimationAsset<FPoseSearchDatabaseAnimationAssetBase>(*SearchIndexAsset))
+					TimelineNormalizedTime = TimelineRealTime;
+
+					// @todo: This is a hack, since we don't have a way to get the normalized timeline of a blendspace yet (pending syncing...?)
+					if (const UE::PoseSearch::FSearchIndexAsset* SearchIndexAsset = MotionMatchingState.CurrentSearchResult.GetSearchIndexAsset())
 					{
-						if (Cast<UBlendSpace>(DatabaseAnimationAssetBase->GetAnimationAssetForRole(MotionMatchingState.CurrentSearchResult.Database->Schema->GetDefaultRole())))
+						if (const FPoseSearchDatabaseAnimationAssetBase* DatabaseAnimationAssetBase = MotionMatchingState.CurrentSearchResult.Database->GetDatabaseAnimationAsset<FPoseSearchDatabaseAnimationAssetBase>(*SearchIndexAsset))
 						{
-							// Convert to normalized time.
-							const float TimelineDuration = TimelineState.GetDuration();
-							if (TimelineDuration > 0.0f)
+							if (Cast<UBlendSpace>(DatabaseAnimationAssetBase->GetAnimationAssetForRole(MotionMatchingState.CurrentSearchResult.Database->Schema->GetDefaultRole())))
 							{
-								// NOTE: This doesn't work because AnimNext normalized time does not much database indexing normalized time.
-								TimelinePosition /= TimelineState.GetDuration();
+								// Convert to normalized time.
+								const float TimelineDuration = TimelineState.GetDuration();
+								if (TimelineDuration > 0.0f)
+								{
+									// NOTE: This doesn't work because AnimNext normalized time does not much database indexing normalized time.
+									TimelineNormalizedTime /= TimelineDuration;
+								}
 							}
 						}
 					}
+
+					MotionMatchingState.CurrentSearchResult.UpdateWithNormalizedTime(TimelineNormalizedTime);
 				}
 
-				MotionMatchingState.AdjustAssetTime(TimelinePosition);
+#if ENABLE_VISUAL_LOG
+				if (FVisualLogger::IsRecording())
+				{
+					//const float AfterNormalize = TimelinePosition;
+					const float AfterMotionMatchingStateStuff = MotionMatchingState.CurrentSearchResult.AssetTime;
+					const FName AnimName = TimelineState.GetDebugName();
+					const float Duration = TimelineState.GetDuration();
+					static const TCHAR* LogName = TEXT("FMotionMatchingTrait");
+					UE_VLOG(Context.GetHostObject(), LogName, Verbose, TEXT("TimelineRealTime: %f"), TimelineRealTime);
+					UE_VLOG(Context.GetHostObject(), LogName, Verbose, TEXT("TimelineNormalizedTime: %f"), TimelineNormalizedTime);
+					UE_VLOG(Context.GetHostObject(), LogName, Verbose, TEXT("AfterMotionMatchingStateStuff: %f"), AfterMotionMatchingStateStuff);
+					UE_VLOG(Context.GetHostObject(), LogName, Verbose, TEXT("AnimName: %s"), *AnimName.ToString());
+					UE_VLOG(Context.GetHostObject(), LogName, Verbose, TEXT("Duration: %f"), Duration);
+				}
+#endif
 			}
 
 			const TConstArrayView<TObjectPtr<const UPoseSearchDatabase>> DatabasesToSearch = SharedData->GetDatabases(Binding);
@@ -319,23 +358,40 @@ namespace UE::AnimNext
 		DesiredPlayRate = GVarAnimNextMotionMatchingPlayRateEnabled == UE::Private::EPlayRateState::Enabled ? DesiredPlayRate : 1.0f;
 #endif // ENABLE_ANIM_DEBUG
 
-		const bool bUpdateBlendSpaceInputs = SharedData->GetbUpdateBlendSpaceInputs(Binding);
 		if (AnimationAsset)
 		{
+			const bool bTrySkipBlendsForBlendSpaces = SharedData->GetbTrySkipBlendsForBlendSpaces(Binding);
 			const bool bIsBlendSpace = AnimationAsset->IsA<UBlendSpace>();
 			bool bSkipBlend = false;
-			if (bIsBlendSpace && bUpdateBlendSpaceInputs)
+			if (bIsBlendSpace && bTrySkipBlendsForBlendSpaces)
 			{
 				IBlendStack::FGraphRequestPtr ActiveGraphRequest;
 				BlendStackBinding.GetActiveGraph(Context, ActiveGraphRequest);
 				if (ActiveGraphRequest && (ActiveGraphRequest->FactoryObject == AnimationAsset))
 				{
-					// If we're still on the same blend space, and we're updating blend space inputs dynamically, then DON'T BLEND.
-					bSkipBlend = true;
+					check(ActiveGraphRequest->AnimationGraph);
+					for (FStructView& Payload : ActiveGraphRequest->GraphPayload.Get())
+					{
+						if (FAnimNextNativeDataInterface_BlendSpacePlayer* BlendSpacePayload = Payload.GetPtr<FAnimNextNativeDataInterface_BlendSpacePlayer>())
+						{
+							const float MaxDeltaAssetTimeToTrySkipBlendsForBlendSpaces = SharedData->GetMaxDeltaAssetTimeToTrySkipBlendsForBlendSpaces(Binding);
+							if (FMath::Abs(BlendSpacePayload->StartPosition - AssetTime) < MaxDeltaAssetTimeToTrySkipBlendsForBlendSpaces)
+							{
+								// If we're still on the same blend space, and we're updating blend space inputs dynamically, then DON'T BLEND.
+								bSkipBlend = true;
+
+								BlendSpacePayload->PlayRate = DesiredPlayRate;
+
+								BlendSpacePayload->XAxisSamplePoint = BlendSpaceParameters.X;
+								BlendSpacePayload->YAxisSamplePoint = BlendSpaceParameters.Y;
+							}
+							break;
+						}
+					}
 				}
 			}
 
-			if (bSkipBlend == false)
+			if (!bSkipBlend)
 			{
 				UE::AnimNext::IBlendStack::FGraphRequest NewGraphRequest;
 				NewGraphRequest.BlendArgs = SharedData->BlendArguments;
@@ -351,7 +407,7 @@ namespace UE::AnimNext
 					auto& NativeInterface = SyncGroupPayload.InitializeAs<FAnimNextNativeDataInterface_SynchronizeUsingGroups>();
 					NativeInterface.GroupName = NAME_None;
 					NativeInterface.GroupRole = EAnimGroupSynchronizationRole::AlwaysFollower;
-					NativeInterface.SyncMode = EAnimGroupSynchronizationMode::SynchronizeUsingUniqueGroupName;
+					NativeInterface.SyncMode = SharedData->GetSyncMode(Binding);
 					NativeInterface.MatchSyncPoint = true;
 					GraphPayload.AddNative(MoveTemp(SyncGroupPayload));
 
@@ -405,7 +461,9 @@ namespace UE::AnimNext
 				else if (FAnimNextNativeDataInterface_BlendSpacePlayer* BlendSpacePayload = Payload.GetPtr<FAnimNextNativeDataInterface_BlendSpacePlayer>())
 				{
 					BlendSpacePayload->PlayRate = DesiredPlayRate;
-					if (bUpdateBlendSpaceInputs)
+					
+					// always updating the blend paramters if bUpdateBlendSpaceInputs
+					if (SharedData->GetbUpdateBlendSpaceInputs(Binding))
 					{
 						BlendSpacePayload->XAxisSamplePoint = SharedData->GetXAxisSamplePoint(Binding);
 						BlendSpacePayload->YAxisSamplePoint = SharedData->GetYAxisSamplePoint(Binding);
diff --git a/Engine/Plugins/Experimental/AnimNextPoseSearch/Source/AnimNextPoseSearch/Private/MotionMatchingTrait.h b/Engine/Plugins/Experimental/AnimNextPoseSearch/Source/AnimNextPoseSearch/Private/MotionMatchingTrait.h
--- a/Engine/Plugins/Experimental/AnimNextPoseSearch/Source/AnimNextPoseSearch/Private/MotionMatchingTrait.h
+++ b/Engine/Plugins/Experimental/AnimNextPoseSearch/Source/AnimNextPoseSearch/Private/MotionMatchingTrait.h
@@ -3,13 +3,14 @@
 #pragma once
 
 #include "CoreMinimal.h"
+#include "Module/ModuleHandle.h"
 #include "PoseSearch/PoseSearchInteractionLibrary.h"
 #include "PoseSearch/PoseSearchLibrary.h"
 #include "TraitCore/Trait.h"
 #include "TraitCore/TraitSharedData.h"
 #include "TraitInterfaces/IEvaluate.h"
+#include "TraitInterfaces/IGroupSynchronization.h"
 #include "TraitInterfaces/IUpdate.h"
-#include "Module/ModuleHandle.h"
 #include "MotionMatchingTrait.generated.h"
 
 USTRUCT(meta = (DisplayName = "Motion Matching"))
@@ -61,6 +62,17 @@ struct FMotionMatchingTraitSharedData : public FAnimNextTraitSharedData
 	UPROPERTY(EditAnywhere, Category = Settings)
 	bool bUpdateBlendSpaceInputs = false;
 
+	// if bTrySkipBlendsForBlendSpaces is true and blendstack is already playing the currently selected blendspace, MM trait will try reusing the same graph by only changing the blend parameters
+	UPROPERTY(EditAnywhere, Category = Settings)
+	bool bTrySkipBlendsForBlendSpaces = true;
+
+	UPROPERTY(EditAnywhere, Category = Settings)
+	float MaxDeltaAssetTimeToTrySkipBlendsForBlendSpaces = 0.1f;
+
+	// The synchronization mode
+	UPROPERTY(EditAnywhere, Category = Settings)
+	EAnimGroupSynchronizationMode SyncMode = EAnimGroupSynchronizationMode::SynchronizeUsingUniqueGroupName;
+
 	// @todo: should this be in FMotionMatchingState?
 	UPROPERTY(EditAnywhere, Category = Settings, meta = (PinHiddenByDefault))
 	// Applied EPoseSearchInterruptMode on the next update that controls the continuing pose search eveluation. This is set back to EPoseSearchInterruptMode::DoNotInterrupt after each update.
@@ -102,6 +114,9 @@ struct FMotionMatchingTraitSharedData : public FAnimNextTraitSharedData
 		GeneratorMacro(PlayRate) \
 		GeneratorMacro(bShouldUseCachedChannelData) \
 		GeneratorMacro(bUpdateBlendSpaceInputs) \
+		GeneratorMacro(bTrySkipBlendsForBlendSpaces) \
+		GeneratorMacro(MaxDeltaAssetTimeToTrySkipBlendsForBlendSpaces) \
+		GeneratorMacro(SyncMode) \
 		GeneratorMacro(NextUpdateInterruptMode) \
 		GeneratorMacro(MotionMatchingResultVariableName) \
 		GeneratorMacro(MotionMatchingResultModuleHandle) \
