PoseSearch - MotionMatching state pose index deprecation to improve continuing pose calculation

diff --git a/Engine/Plugins/Animation/BlendStack/Source/Runtime/Private/AnimNode_BlendStack.cpp b/Engine/Plugins/Animation/BlendStack/Source/Runtime/Private/AnimNode_BlendStack.cpp
--- a/Engine/Plugins/Animation/BlendStack/Source/Runtime/Private/AnimNode_BlendStack.cpp
+++ b/Engine/Plugins/Animation/BlendStack/Source/Runtime/Private/AnimNode_BlendStack.cpp
@@ -948,6 +948,16 @@ float FAnimNode_BlendStack_Standalone::GetAccumulatedTime() const
 	return AnimPlayers.IsEmpty() ? 0.f : AnimPlayers[0].GetAccumulatedTime();
 }
 
+bool FAnimNode_BlendStack_Standalone::GetMirror() const
+{
+	return AnimPlayers.IsEmpty() ? false : AnimPlayers[0].GetMirror();
+}
+
+FVector FAnimNode_BlendStack_Standalone::GetBlendParameters() const
+{
+	return AnimPlayers.IsEmpty() ? FVector::ZeroVector : AnimPlayers[0].GetBlendParameters();
+}
+
 static void RequestInertialBlend(const FAnimationUpdateContext& Context, float BlendTime, const UBlendProfile* BlendProfile, EAlphaBlendOption BlendOption, FName InertialBlendNodeTag)
 {
 #if ENABLE_ANIM_DEBUG
@@ -1021,6 +1031,8 @@ void FAnimNode_BlendStack_Standalone::BlendTo(const FAnimationUpdateContext& Con
 				IPoseSearchProvider::FSearchPlayingAsset PlayingAsset;
 				PlayingAsset.Asset = GetAnimAsset();
 				PlayingAsset.AccumulatedTime = GetAccumulatedTime();
+				PlayingAsset.bMirrored = GetMirror();
+				PlayingAsset.BlendParameters = GetBlendParameters();
 
 				IPoseSearchProvider::FSearchFutureAsset FutureAsset;
 				FutureAsset.Asset = AnimationAsset;
diff --git a/Engine/Plugins/Animation/BlendStack/Source/Runtime/Public/BlendStack/AnimNode_BlendStack.h b/Engine/Plugins/Animation/BlendStack/Source/Runtime/Public/BlendStack/AnimNode_BlendStack.h
--- a/Engine/Plugins/Animation/BlendStack/Source/Runtime/Public/BlendStack/AnimNode_BlendStack.h
+++ b/Engine/Plugins/Animation/BlendStack/Source/Runtime/Public/BlendStack/AnimNode_BlendStack.h
@@ -182,6 +182,9 @@ struct BLENDSTACK_API FAnimNode_BlendStack_Standalone : public FAnimNode_AssetPl
 	virtual float GetCurrentAssetLength() const override;
 	// End of FAnimNode_AssetPlayerBase interface
 
+	bool GetMirror() const;
+	FVector GetBlendParameters() const;
+
 	int32 GetNextPoseLinkIndex();
 
 	int32 GetMaxActiveBlends() const { return MaxActiveBlends; }
diff --git a/Engine/Plugins/Animation/PoseSearch/Source/Editor/Private/PoseSearchDatabaseEditorReflection.cpp b/Engine/Plugins/Animation/PoseSearch/Source/Editor/Private/PoseSearchDatabaseEditorReflection.cpp
--- a/Engine/Plugins/Animation/PoseSearch/Source/Editor/Private/PoseSearchDatabaseEditorReflection.cpp
+++ b/Engine/Plugins/Animation/PoseSearch/Source/Editor/Private/PoseSearchDatabaseEditorReflection.cpp
@@ -167,15 +167,6 @@ bool UPoseSearchDatabaseAnimMontageReflection::ApplyChanges()
 
 bool UPoseSearchDatabaseMultiAnimAssetReflection::ApplyChanges()
 {
-	// @todo: add support for InteractionAsset contatining blend spaces
-	FVector BlendParameters = FVector::ZeroVector;
-	const float PlayLength = MultiAnimAsset.GetPlayLength(BlendParameters);
-	const FFloatInterval ClampedSamplingRange(FMath::Clamp(MultiAnimAsset.SamplingRange.Min, 0.0f, PlayLength), FMath::Clamp(MultiAnimAsset.SamplingRange.Max, 0.0f, PlayLength));
-	if (ClampedSamplingRange != MultiAnimAsset.SamplingRange)
-	{
-		MultiAnimAsset.SamplingRange = ClampedSamplingRange;
-	}
-
 	if (const TSharedPtr<UE::PoseSearch::FDatabaseAssetTreeNode> AssetTreeNode = WeakAssetTreeNode.Pin())
 	{
 		if (const TSharedPtr<UE::PoseSearch::FDatabaseViewModel> ViewModel = AssetTreeNode->EditorViewModel.Pin())
diff --git a/Engine/Plugins/Animation/PoseSearch/Source/Editor/Private/PoseSearchDatabaseViewModel.cpp b/Engine/Plugins/Animation/PoseSearch/Source/Editor/Private/PoseSearchDatabaseViewModel.cpp
--- a/Engine/Plugins/Animation/PoseSearch/Source/Editor/Private/PoseSearchDatabaseViewModel.cpp
+++ b/Engine/Plugins/Animation/PoseSearch/Source/Editor/Private/PoseSearchDatabaseViewModel.cpp
@@ -214,7 +214,7 @@ void FDatabasePreviewActor::UpdatePreviewActor(const UPoseSearchDatabase* PoseSe
 	FAnimationRuntime::AdvanceTime(false, CurrentPlayTime, CurrentTime, IndexAsset.GetLastSampleTime(PoseSearchDatabase->Schema->SampleRate));
 			 
 	// time to pose index
-	CurrentPoseIndex = PoseSearchDatabase->GetPoseIndexFromTime(CurrentTime, IndexAsset);
+	CurrentPoseIndex = IndexAsset.GetPoseIndexFromTime(CurrentTime, PoseSearchDatabase->Schema->SampleRate);
 
 	QuantizedTime = CurrentPoseIndex >= 0 ? PoseSearchDatabase->GetRealAssetTime(CurrentPoseIndex) : CurrentTime;
 	if (bQuantizeAnimationToPoseData)
diff --git a/Engine/Plugins/Animation/PoseSearch/Source/Editor/Private/PoseSearchInteractionAssetEditor.cpp b/Engine/Plugins/Animation/PoseSearch/Source/Editor/Private/PoseSearchInteractionAssetEditor.cpp
--- a/Engine/Plugins/Animation/PoseSearch/Source/Editor/Private/PoseSearchInteractionAssetEditor.cpp
+++ b/Engine/Plugins/Animation/PoseSearch/Source/Editor/Private/PoseSearchInteractionAssetEditor.cpp
@@ -226,10 +226,7 @@ void FInteractionAssetViewModel::PreviewBackwardStep()
 {
 	if (const UPoseSearchInteractionAsset* InteractionAsset = GetInteractionAsset())
 	{
-		// @todo: add support for InteractionAsset contatining blend spaces
-		FVector BlendParameters = FVector::ZeroVector;
-		const float NewPlayTime = FMath::Clamp(PlayTime - StepDeltaTime, 0.f, InteractionAsset->GetPlayLength(BlendParameters));
-		SetPlayTime(NewPlayTime, false);
+		SetPlayTime(PlayTime - StepDeltaTime, false);
 	}
 }
 
@@ -252,10 +249,7 @@ void FInteractionAssetViewModel::PreviewForwardStep()
 {
 	if (const UPoseSearchInteractionAsset* InteractionAsset = GetInteractionAsset())
 	{
-		// @todo: add support for InteractionAsset contatining blend spaces
-		FVector BlendParameters = FVector::ZeroVector;
-		const float NewPlayTime = FMath::Clamp(PlayTime + StepDeltaTime, 0.f, InteractionAsset->GetPlayLength(BlendParameters));
-		SetPlayTime(NewPlayTime, false);
+		SetPlayTime(PlayTime + StepDeltaTime, false);
 	}
 }
 
@@ -263,9 +257,8 @@ void FInteractionAssetViewModel::PreviewForwardEnd()
 {
 	if (const UPoseSearchInteractionAsset* InteractionAsset = GetInteractionAsset())
 	{
-		// @todo: add support for InteractionAsset contatining blend spaces
-		FVector BlendParameters = FVector::ZeroVector;
-		SetPlayTime(InteractionAsset->GetPlayLength(BlendParameters), false);
+		// setting play time to a big number that will be clamped internally
+		SetPlayTime(UE_MAX_FLT, false);
 	}
 }
 
@@ -285,9 +278,6 @@ void FInteractionAssetViewModel::Tick(float DeltaSeconds)
 	}
 
 	PlayTime += DeltaSeconds * DeltaTimeMultiplier;
-	// @todo: add support for InteractionAsset contatining blend spaces
-	FVector BlendParameters = FVector::ZeroVector;
-	PlayTime = FMath::Clamp(PlayTime, 0.f, InteractionAsset->GetPlayLength(BlendParameters));
 
 	FRoleToIndex InteractionAssetRoleToIndex;
 	for (int32 RoleIndex = 0; RoleIndex < InteractionAsset->GetNumRoles(); ++RoleIndex)
@@ -404,9 +394,7 @@ void FInteractionAssetViewModel::SetPlayTime(float NewPlayTime, bool bInTickPlay
 {
 	if (const UPoseSearchInteractionAsset* InteractionAsset = GetInteractionAsset())
 	{
-		// @todo: add support for InteractionAsset contatining blend spaces
-		FVector BlendParameters = FVector::ZeroVector;
-		NewPlayTime = FMath::Clamp(NewPlayTime, 0.f, InteractionAsset->GetPlayLength(BlendParameters));
+		NewPlayTime = FMath::Max(NewPlayTime, 0.f);
 		DeltaTimeMultiplier = bInTickPlayTime ? DeltaTimeMultiplier : 0.f;
 
 		if (!FMath::IsNearlyEqual(PlayTime, NewPlayTime))
diff --git a/Engine/Plugins/Animation/PoseSearch/Source/Runtime/Private/AnimNode_MotionMatching.cpp b/Engine/Plugins/Animation/PoseSearch/Source/Runtime/Private/AnimNode_MotionMatching.cpp
--- a/Engine/Plugins/Animation/PoseSearch/Source/Runtime/Private/AnimNode_MotionMatching.cpp
+++ b/Engine/Plugins/Animation/PoseSearch/Source/Runtime/Private/AnimNode_MotionMatching.cpp
@@ -89,36 +89,25 @@ void FAnimNode_MotionMatching::UpdateAssetPlayer(const FAnimationUpdateContext&
 		MotionMatchingState.Reset();
 		FAnimNode_BlendStack_Standalone::Reset();
 	}
-	else if (MotionMatchingState.CurrentSearchResult.Database == nullptr || MotionMatchingState.CurrentSearchResult.Database->Schema == nullptr)
+	else if (MotionMatchingState.SearchResult.SelectedDatabase == nullptr || MotionMatchingState.SearchResult.SelectedDatabase->Schema == nullptr)
 	{
 	}
 #if WITH_EDITOR
-	else if (EAsyncBuildIndexResult::Success != FAsyncPoseSearchDatabasesManagement::RequestAsyncBuildIndex(MotionMatchingState.CurrentSearchResult.Database.Get(), ERequestAsyncBuildFlag::ContinueRequest))
+	else if (EAsyncBuildIndexResult::Success != FAsyncPoseSearchDatabasesManagement::RequestAsyncBuildIndex(MotionMatchingState.SearchResult.SelectedDatabase.Get(), ERequestAsyncBuildFlag::ContinueRequest))
 	{
-		// MotionMatchingState.CurrentSearchResult.Database is indexing, and it's not safe to use its previous index properties cached in MotionMatchingState
-		MotionMatchingState.Reset();
-	}
-	else if (MotionMatchingState.CurrentSearchResult.IsValid() && MotionMatchingState.CurrentSearchResult.PoseIdx >= MotionMatchingState.CurrentSearchResult.Database->GetSearchIndex().GetNumPoses())
-	{
-		// in case this node is not updated, and MotionMatchingState.CurrentSearchResult.Database gets modified, we could end up with
-		// CurrentSearchResult being out of synch with the updated database, so we need to reset the state
+		// MotionMatchingState.SearchResult.Database is indexing, and it's not safe to use its previous index properties cached in MotionMatchingState
 		MotionMatchingState.Reset();
 	}
 #endif // WITH_EDITOR
-	else if (GetAnimAsset() != MotionMatchingState.CurrentSearchResult.GetCurrentResultAnimationAsset())
-	{
-		// MotionMatchingState and blendstack are not in synch! MotionMatchingState thinks it's playing MotionMatchingState.CurrentSearchResult.GetCurrentResultAnimationAsset()
-		// and blendstack thinks it's playing GetAnimAsset(). The two could have gone out of synch becasue or external blendstack requests, and so we reset the MotionMatchingState
-		MotionMatchingState.Reset();
-	}
 	else
 	{
 		// We adjust the motion matching state asset time to the current player node's asset time. This is done 
 		// because the player node may have ticked more or less time than we expected due to variable dt or the 
 		// dynamic playback rate adjustment and as such the motion matching state does not update by itself
-		
-		// @todo: this is missing eventual update of the blend parameters that could change the pose!
-		MotionMatchingState.CurrentSearchResult.UpdateWithNormalizedTime(GetAccumulatedTime());
+		MotionMatchingState.SearchResult.SelectedAnim = GetAnimAsset();
+		MotionMatchingState.SearchResult.SelectedTime = GetAccumulatedTime();
+		MotionMatchingState.SearchResult.bIsMirrored = GetMirror();
+		MotionMatchingState.SearchResult.BlendParameters = GetBlendParameters();
 	}
 
 	UpdateCounter.SynchronizeWith(Context.AnimInstanceProxy->GetUpdateCounter());
@@ -136,12 +125,9 @@ void FAnimNode_MotionMatching::UpdateAssetPlayer(const FAnimationUpdateContext&
 #if ENABLE_ANIM_DEBUG
 	if (GVarAnimNodeMotionMatchingDrawInfo)
 	{
-		const UPoseSearchDatabase* CurrentDatabase = MotionMatchingState.CurrentSearchResult.Database.Get();
-		const UAnimationAsset* CurrentAnimationAsset = GetAnimAsset();
-
 		FString DebugInfo = FString::Printf(TEXT("NextUpdateInterruptMode(%s)\n"), *UEnum::GetValueAsString(NextUpdateInterruptMode));
-		DebugInfo += FString::Printf(TEXT("Current Database(%s)\n"), *GetNameSafe(CurrentDatabase));
-		DebugInfo += FString::Printf(TEXT("Current Asset(%s)\n"), *GetNameSafe(CurrentAnimationAsset));
+		DebugInfo += FString::Printf(TEXT("Current Database(%s)\n"), *GetNameSafe(MotionMatchingState.SearchResult.SelectedDatabase.Get()));
+		DebugInfo += FString::Printf(TEXT("Current Asset(%s)\n"), *GetNameSafe(GetAnimAsset()));
 		if (GVarAnimNodeMotionMatchingDrawInfoVerbose)
 		{
 			DebugInfo += FString::Printf(TEXT("Databases to search:\n"));
@@ -193,13 +179,13 @@ void FAnimNode_MotionMatching::UpdateAssetPlayer(const FAnimationUpdateContext&
 		, EventToSearch
 	);
 
-	const bool bJumpToPose = MotionMatchingState.CurrentSearchResult.IsValid() && !MotionMatchingState.CurrentSearchResult.bIsContinuingPoseSearch;
+	const bool bJumpToPose = MotionMatchingState.SearchResult.SelectedAnim.IsValid() && !MotionMatchingState.SearchResult.bIsContinuingPoseSearch;
 
 	PRAGMA_DISABLE_DEPRECATION_WARNINGS
 	MotionMatchingState.bJumpedToPose = bJumpToPose;
 	PRAGMA_ENABLE_DEPRECATION_WARNINGS
 
-	float DesiredPlayRate = MotionMatchingState.WantedPlayRate * PlayRateMultiplier;
+	float DesiredPlayRate = MotionMatchingState.SearchResult.WantedPlayRate * PlayRateMultiplier;
 #if ENABLE_ANIM_DEBUG
 	DesiredPlayRate = GVarAnimNodeMotionMatchingPlayRateEnabled == UE::Private::EPlayRateState::Enabled ? DesiredPlayRate : PlayRateMultiplier;
 #endif // ENABLE_ANIM_DEBUG
@@ -209,14 +195,10 @@ void FAnimNode_MotionMatching::UpdateAssetPlayer(const FAnimationUpdateContext&
 	// If a new pose is requested, blend into the new asset via BlendStackNode
 	if (bJumpToPose)
 	{
-		const FSearchIndexAsset* SearchIndexAsset = MotionMatchingState.CurrentSearchResult.GetSearchIndexAsset();
-		const UPoseSearchDatabase* CurrentResultDatabase = MotionMatchingState.CurrentSearchResult.Database.Get();
-		if (SearchIndexAsset && CurrentResultDatabase && CurrentResultDatabase->Schema)
+		const UPoseSearchDatabase* CurrentResultDatabase = MotionMatchingState.SearchResult.SelectedDatabase.Get();
+		if (CurrentResultDatabase && CurrentResultDatabase->Schema)
 		{
-			const FPoseSearchDatabaseAnimationAssetBase* DatabaseAsset = CurrentResultDatabase->GetDatabaseAnimationAsset<FPoseSearchDatabaseAnimationAssetBase>(*SearchIndexAsset);
-			check(DatabaseAsset);
-			
-			if (UAnimationAsset* AnimationAsset = Cast<UAnimationAsset>(DatabaseAsset->GetAnimationAssetForRole(CurrentResultDatabase->Schema->GetDefaultRole())))
+			if (UAnimationAsset* AnimationAsset = MotionMatchingState.SearchResult.GetAnimationAssetForRole())
 			{
 				// Clear up any sync group info before pushing new asset player (which will have sync info since its the highest weighted).
 				for (FBlendStackAnimPlayer& AnimPlayer : AnimPlayers)
@@ -248,9 +230,9 @@ void FAnimNode_MotionMatching::UpdateAssetPlayer(const FAnimationUpdateContext&
 				
 				if (!bAbortBlend)
 				{
-					FAnimNode_BlendStack_Standalone::BlendTo(Context, AnimationAsset, MotionMatchingState.CurrentSearchResult.AssetTime,
-						SearchIndexAsset->IsLooping(), SearchIndexAsset->IsMirrored(), CurrentResultDatabase->Schema->GetMirrorDataTable(DefaultRole), BlendTime,
-						BlendProfile, BlendOption, bUseInertialBlend, NAME_None, SearchIndexAsset->GetBlendParameters(), DesiredPlayRate, 0,
+					FAnimNode_BlendStack_Standalone::BlendTo(Context, AnimationAsset, MotionMatchingState.SearchResult.SelectedTime, MotionMatchingState.SearchResult.bLoop,
+						MotionMatchingState.SearchResult.bIsMirrored, CurrentResultDatabase->Schema->GetMirrorDataTable(MotionMatchingState.SearchResult.Role), BlendTime,
+						BlendProfile, BlendOption, bUseInertialBlend, NAME_None, MotionMatchingState.SearchResult.BlendParameters, DesiredPlayRate, 0,
 						GetGroupName(), GetGroupRole(), GetGroupMethod(), GetOverridePositionWhenJoiningSyncGroupAsLeader());
 				}
 			}
diff --git a/Engine/Plugins/Animation/PoseSearch/Source/Runtime/Private/AnimNode_MotionMatchingInteraction.cpp b/Engine/Plugins/Animation/PoseSearch/Source/Runtime/Private/AnimNode_MotionMatchingInteraction.cpp
--- a/Engine/Plugins/Animation/PoseSearch/Source/Runtime/Private/AnimNode_MotionMatchingInteraction.cpp
+++ b/Engine/Plugins/Animation/PoseSearch/Source/Runtime/Private/AnimNode_MotionMatchingInteraction.cpp
@@ -63,10 +63,12 @@ void FAnimNode_MotionMatchingInteraction::UpdateAssetPlayer(const FAnimationUpda
 	const float DeltaTime = Context.GetDeltaTime();
 	if (FPoseHistoryProvider* PoseHistoryProvider = Context.GetMessage<FPoseHistoryProvider>())
 	{
-		// we're NOT providing the UMultiAnimAsset here, relying on the interaction subsystem to resolve the ContinuingProperties properly
+		// @todo: retrieve the playing UMultiAnimAsset from the UAnimAsset blendspace is playing via GetAnimAsset() and blend parameters via GetBlendParameters() and mirror state
 		FPoseSearchContinuingProperties ContinuingProperties;
-		ContinuingProperties.PlayingAsset = CurrentResult.SelectedAnim.Get(); // GetAnimAsset();
+		ContinuingProperties.PlayingAsset = CurrentResult.SelectedAnim.Get();
 		ContinuingProperties.PlayingAssetAccumulatedTime = GetAccumulatedTime();
+		ContinuingProperties.bIsPlayingAssetMirrored = CurrentResult.bIsMirrored;
+		ContinuingProperties.PlayingAssetBlendParameters = CurrentResult.BlendParameters;
 
 		UPoseSearchInteractionLibrary::MotionMatchInteraction(CurrentResult, Availabilities, Context.AnimInstanceProxy->GetAnimInstanceObject(), ContinuingProperties, FName(),  & PoseHistoryProvider->GetPoseHistory(), bValidateResultAgainstAvailabilities);
 		check(CurrentResult.ActorRootTransforms.Num() == CurrentResult.ActorRootBoneTransforms.Num());
diff --git a/Engine/Plugins/Animation/PoseSearch/Source/Runtime/Private/MotionMatchingAnimNodeLibrary.cpp b/Engine/Plugins/Animation/PoseSearch/Source/Runtime/Private/MotionMatchingAnimNodeLibrary.cpp
--- a/Engine/Plugins/Animation/PoseSearch/Source/Runtime/Private/MotionMatchingAnimNodeLibrary.cpp
+++ b/Engine/Plugins/Animation/PoseSearch/Source/Runtime/Private/MotionMatchingAnimNodeLibrary.cpp
@@ -24,49 +24,16 @@ void UMotionMatchingAnimNodeLibrary::GetMotionMatchingSearchResult(const FMotion
 {
 	using namespace UE::PoseSearch;
 	
-	Result = FPoseSearchBlueprintResult();
-	bIsResultValid = false;
-
 	if (FAnimNode_MotionMatching* MotionMatchingNodePtr = MotionMatchingNode.GetAnimNodePtr<FAnimNode_MotionMatching>())
 	{
 		const FMotionMatchingState& MotionMatchingState = MotionMatchingNodePtr->GetMotionMatchingState();
-		const UPoseSearchDatabase* CurrentResultDatabase = MotionMatchingState.CurrentSearchResult.Database.Get();
-		if (CurrentResultDatabase == nullptr || CurrentResultDatabase->Schema == nullptr)
-		{
-		}
-#if WITH_EDITOR
-		else if (EAsyncBuildIndexResult::Success != FAsyncPoseSearchDatabasesManagement::RequestAsyncBuildIndex(CurrentResultDatabase, ERequestAsyncBuildFlag::ContinueRequest))
-		{
-			// MotionMatchingState.CurrentSearchResult.Database is indexing, and it's not safe to use its previous index properties cached in MotionMatchingState
-			UE_LOG(LogPoseSearch, Warning, TEXT("UMotionMatchingAnimNodeLibrary::GetMotionMatchingSearchResult called while '%s' is indexing. returning an invalid result.."), *CurrentResultDatabase->GetName());
-		}
-		else if (MotionMatchingState.CurrentSearchResult.IsValid() && MotionMatchingState.CurrentSearchResult.PoseIdx >= MotionMatchingState.CurrentSearchResult.Database->GetSearchIndex().GetNumPoses())
-		{
-			// in case this node is not updated, and MotionMatchingState.CurrentSearchResult.Database gets modified, we could end up with 
-			// CurrentSearchResult being out of synch with the updated database, so we need to reset the state
-			UE_LOG(LogPoseSearch, Warning, TEXT("UMotionMatchingAnimNodeLibrary::GetMotionMatchingSearchResult inconsistency between MotionMatchingState and '%s', probably because of stale MotionMatchingState after database got reindexed. returning an invalid result.."), *CurrentResultDatabase->GetName());
-		}
-#endif // WITH_EDITOR
-		else if (const FSearchIndexAsset* SearchIndexAsset = MotionMatchingState.CurrentSearchResult.GetSearchIndexAsset())
-		{
-			if (const FPoseSearchDatabaseAnimationAssetBase* DatabaseAsset = CurrentResultDatabase->GetDatabaseAnimationAsset<FPoseSearchDatabaseAnimationAssetBase>(*SearchIndexAsset))
-			{
-				Result.SelectedAnim = DatabaseAsset->GetAnimationAsset();
-				Result.SelectedTime = MotionMatchingState.CurrentSearchResult.AssetTime;
-				Result.bIsContinuingPoseSearch = MotionMatchingState.CurrentSearchResult.bIsContinuingPoseSearch;
-				Result.WantedPlayRate = MotionMatchingState.WantedPlayRate;
-				Result.bLoop = SearchIndexAsset->IsLooping();
-				Result.bIsMirrored = SearchIndexAsset->IsMirrored();
-				Result.BlendParameters = SearchIndexAsset->GetBlendParameters();
-				Result.SelectedDatabase = CurrentResultDatabase;
-				Result.SearchCost = MotionMatchingState.CurrentSearchResult.PoseCost;
-				Result.Role = CurrentResultDatabase->Schema->GetDefaultRole();
-				bIsResultValid = true;
-			}
-		}
+		Result = MotionMatchingState.SearchResult;
+		bIsResultValid = MotionMatchingState.SearchResult.SelectedAnim != nullptr;
 	}
 	else
 	{
+		Result = FPoseSearchBlueprintResult();
+		bIsResultValid = false;
 		UE_LOG(LogPoseSearch, Warning, TEXT("UMotionMatchingAnimNodeLibrary::GetMotionMatchingSearchResult called on an invalid context or with an invalid type"));
 	}
 }
diff --git a/Engine/Plugins/Animation/PoseSearch/Source/Runtime/Private/PoseSearchDatabase.cpp b/Engine/Plugins/Animation/PoseSearch/Source/Runtime/Private/PoseSearchDatabase.cpp
--- a/Engine/Plugins/Animation/PoseSearch/Source/Runtime/Private/PoseSearchDatabase.cpp
+++ b/Engine/Plugins/Animation/PoseSearch/Source/Runtime/Private/PoseSearchDatabase.cpp
@@ -796,7 +796,7 @@ TConstArrayView<int32> UPoseSearchDatabase::GetAssetIndexesForSourceAsset(const
 	{
 #if DO_CHECK
 		// validating the consistency of IndexesForSourceAsset retrieved from SourceAsset
-		const UE::PoseSearch::FSearchIndex& SearchIndex = GetSearchIndex();
+		const FSearchIndex& SearchIndex = GetSearchIndex();
 		for (int32 AssetIndex : *IndexesForSourceAsset)
 		{
 			const FSearchIndexAsset& SearchIndexAsset = SearchIndex.Assets[AssetIndex];
@@ -868,9 +868,49 @@ const UE::PoseSearch::FSearchIndex& UPoseSearchDatabase::GetSearchIndex() const
 	return SearchIndexPrivate;
 }
 
-int32 UPoseSearchDatabase::GetPoseIndexFromTime(float Time, const UE::PoseSearch::FSearchIndexAsset& SearchIndexAsset) const
+int32 UPoseSearchDatabase::GetPoseIndexFromTime(float RealTimeInSeconds, const UE::PoseSearch::FSearchIndexAsset& SearchIndexAsset) const
 {
-	return SearchIndexAsset.GetPoseIndexFromTime(Time, Schema->SampleRate);
+	return SearchIndexAsset.GetPoseIndexFromTime(RealTimeInSeconds, Schema->SampleRate);
+}
+
+int32 UPoseSearchDatabase::GetPoseIndex(const UObject* AnimationAsset, float AnimationAssetTime, bool bMirrored, const FVector& BlendParameters) const
+{
+	using namespace UE::PoseSearch;
+
+	int32 PoseIdx = INDEX_NONE;
+
+	if (AnimationAsset)
+	{
+		float MinSquaredLength = UE_MAX_FLT;
+		const float SampleRate = Schema->SampleRate;
+		const TConstArrayView<int32> AssetIndexesForSourceAsset = GetAssetIndexesForSourceAsset(AnimationAsset);
+
+		const FSearchIndex& SearchIndex = GetSearchIndex();
+		for (int32 AssetIndex : AssetIndexesForSourceAsset)
+		{
+			const FSearchIndexAsset& SearchIndexAsset = SearchIndex.Assets[AssetIndex];
+			if (SearchIndexAsset.IsMirrored() == bMirrored)
+			{
+				const float BlendParametersSquaredLength = (BlendParameters - SearchIndexAsset.GetBlendParameters()).SquaredLength();
+
+				// using <= so we don't have to check for PoseIdx == INDEX_NONE, since any float will be smaller or equal than UE_MAX_FLT
+				if (BlendParametersSquaredLength <= MinSquaredLength)
+				{
+					MinSquaredLength = BlendParametersSquaredLength;
+
+					const FPoseSearchDatabaseAnimationAssetBase* DatabaseAnimationAssetBase = GetDatabaseAnimationAsset<FPoseSearchDatabaseAnimationAssetBase>(SearchIndexAsset);
+
+					check(DatabaseAnimationAssetBase);
+					check(DatabaseAnimationAssetBase->GetAnimationAsset() == AnimationAsset);
+
+					const float RealAssetTime = AnimationAssetTime * SearchIndexAsset.GetToRealTimeFactor();
+					PoseIdx = SearchIndexAsset.GetPoseIndexFromTime(RealAssetTime, SampleRate);
+				}
+			}
+		}
+	}
+
+	return PoseIdx;
 }
 
 void UPoseSearchDatabase::AddAnimationAsset(FInstancedStruct AnimationAsset)
@@ -1626,6 +1666,7 @@ UE::PoseSearch::FSearchResult UPoseSearchDatabase::SearchContinuingPose(UE::Pose
 	}
 #endif // WITH_EDITOR && ENABLE_ANIM_DEBUG && UE_POSE_SEARCH_TRACE_ENABLED
 
+	check(Result.DebugValidate());
 	return Result;
 }
 
@@ -1898,6 +1939,7 @@ UE::PoseSearch::FSearchResult UPoseSearchDatabase::SearchPCAKDTree(UE::PoseSearc
 		Result.Database = this;
 	}
 
+	check(Result.DebugValidate());
 	return Result;
 }
 
@@ -2001,6 +2043,7 @@ UE::PoseSearch::FSearchResult UPoseSearchDatabase::SearchVPTree(UE::PoseSearch::
 		Result.Database = this;
 	}
 
+	check(Result.DebugValidate());
 	return Result;
 }
 
@@ -2167,6 +2210,7 @@ UE::PoseSearch::FSearchResult UPoseSearchDatabase::SearchBruteForce(UE::PoseSear
 		Result.Database = this;
 	}
 
+	check(Result.DebugValidate());
 	return Result;
 }
 
@@ -2252,4 +2296,4 @@ UE::PoseSearch::FSearchResult UPoseSearchDatabase::SearchEvent(UE::PoseSearch::F
 	}
 
 	return Result;
-}
+}
\ No newline at end of file
diff --git a/Engine/Plugins/Animation/PoseSearch/Source/Runtime/Private/PoseSearchInteractionIsland.cpp b/Engine/Plugins/Animation/PoseSearch/Source/Runtime/Private/PoseSearchInteractionIsland.cpp
--- a/Engine/Plugins/Animation/PoseSearch/Source/Runtime/Private/PoseSearchInteractionIsland.cpp
+++ b/Engine/Plugins/Animation/PoseSearch/Source/Runtime/Private/PoseSearchInteractionIsland.cpp
@@ -1083,6 +1083,7 @@ bool FInteractionIsland::GetResult_AnyThread(const UObject* AnimContext, FPoseSe
 
 			if (bAnimContextFound)
 			{
+				// @todo: perhaps add a custom Result.InitFrom(SearchResult, 1.f) for MM interactions
 				const UPoseSearchDatabase* Database = SearchResult.Database.Get();
 				check(Database);
 
@@ -1092,8 +1093,6 @@ bool FInteractionIsland::GetResult_AnyThread(const UObject* AnimContext, FPoseSe
 				const FPoseSearchDatabaseAnimationAssetBase* DatabaseAnimationAssetBase = Database->GetDatabaseAnimationAsset<FPoseSearchDatabaseAnimationAssetBase>(*SearchIndexAsset);
 				check(DatabaseAnimationAssetBase);
 
-				const FRole Role = SearchContext.Roles[AnimContextIndex];
-
 				Result.SelectedAnim = DatabaseAnimationAssetBase->GetAnimationAsset();
 				Result.SelectedTime = SearchResult.AssetTime;
 				Result.bIsContinuingPoseSearch = SearchResult.bIsContinuingPoseSearch;
@@ -1103,7 +1102,7 @@ bool FInteractionIsland::GetResult_AnyThread(const UObject* AnimContext, FPoseSe
 				Result.SelectedDatabase = Database;
 				Result.SearchCost = SearchResult.PoseCost;
 				Result.bIsInteraction = true;
-				Result.Role = Role;
+				Result.Role = SearchContext.Roles[AnimContextIndex];
 
 				// figuring out the WantedPlayRate
 				Result.WantedPlayRate = 1.f;
diff --git a/Engine/Plugins/Animation/PoseSearch/Source/Runtime/Private/PoseSearchInteractionSubsystem.cpp b/Engine/Plugins/Animation/PoseSearch/Source/Runtime/Private/PoseSearchInteractionSubsystem.cpp
--- a/Engine/Plugins/Animation/PoseSearch/Source/Runtime/Private/PoseSearchInteractionSubsystem.cpp
+++ b/Engine/Plugins/Animation/PoseSearch/Source/Runtime/Private/PoseSearchInteractionSubsystem.cpp
@@ -409,10 +409,14 @@ void UPoseSearchInteractionSubsystem::PopulateContinuingProperties(float DeltaSe
 					{
 						if (const FPoseSearchDatabaseAnimationAssetBase* DatabaseAsset = SearchResult->Database->GetDatabaseAnimationAsset<FPoseSearchDatabaseAnimationAssetBase>(*SearchIndexAsset))
 						{
+							check(SearchIndexAsset->GetToRealTimeFactor() > UE_KINDA_SMALL_NUMBER);
+							// in case DatabaseAsset->GetAnimationAsset() is a blendspace, SearchResult->AssetTime is a normalized time in the interval [0,1] 
+							// so we need to convert the delta time in seconds to the asset normalized time before integrating SearchResult->AssetTime
+							const float NormalizedDeltaTime = DeltaSeconds / SearchIndexAsset->GetToRealTimeFactor();
+							SearchContext.PlayingAssetAccumulatedTime = SearchResult->AssetTime + NormalizedDeltaTime;
 							SearchContext.PlayingAsset = DatabaseAsset->GetAnimationAsset();
-
-							// @todo: for blendspaces PlayingAssetAccumulatedTime should be incremented by a normalized DeltaSeconds
-							SearchContext.PlayingAssetAccumulatedTime = SearchResult->AssetTime + DeltaSeconds;
+							SearchContext.bIsPlayingAssetMirrored = SearchIndexAsset->IsMirrored();
+							SearchContext.PlayingAssetBlendParameters = SearchIndexAsset->GetBlendParameters();
 							// @todo: populate SearchContext.InterruptMode
 						}
 					}
diff --git a/Engine/Plugins/Animation/PoseSearch/Source/Runtime/Private/PoseSearchLibrary.cpp b/Engine/Plugins/Animation/PoseSearch/Source/Runtime/Private/PoseSearchLibrary.cpp
--- a/Engine/Plugins/Animation/PoseSearch/Source/Runtime/Private/PoseSearchLibrary.cpp
+++ b/Engine/Plugins/Animation/PoseSearch/Source/Runtime/Private/PoseSearchLibrary.cpp
@@ -295,6 +295,67 @@ namespace UE::PoseSearch
 
 		return Role;
 	}
+
+	static float CalculateWantedPlayRate(const UE::PoseSearch::FSearchResult& CurrentSearchResult, const UE::PoseSearch::FSearchContext& SearchContext, const FFloatInterval& PlayRate, float TrajectorySpeedMultiplier, const FPoseSearchEvent& EventToSearch)
+	{
+		float WantedPlayRate = 1.f;
+
+		if (CurrentSearchResult.IsValid())
+		{
+			if (CurrentSearchResult.IsEventSearchResult())
+			{
+				// checking if CurrentSearchResult.EventPoseIdx is part of the EventToSearch.EventTag.
+				// If not, it's an event from a continuing pose search that hasn't been interrupted,
+				// so we keep the previously calculated WantedPlayRate
+				const bool bIsEventSearchFromTag = CurrentSearchResult.IsEventSearchFromTag(EventToSearch.EventTag);
+				if (bIsEventSearchFromTag)
+				{
+					const float TimeToEvent = CurrentSearchResult.CalculateTimeToEvent();
+					if (TimeToEvent > UE_KINDA_SMALL_NUMBER && EventToSearch.TimeToEvent > UE_KINDA_SMALL_NUMBER)
+					{
+						// EventToSearch.TimeToEvent is the desired time to event, and TimeToEvent is the actually current time to event. we calculate WantedPlayRate as ratio between the two
+						WantedPlayRate = TimeToEvent / EventToSearch.TimeToEvent;
+					}
+					// if we passed the event (TimeToEvent <= 0) we leave the WantedPlayRate as previously calculated
+				}
+			}
+			else if (!ensure(PlayRate.Min <= PlayRate.Max && PlayRate.Min > UE_KINDA_SMALL_NUMBER))
+			{
+				UE_LOG(LogPoseSearch, Error, TEXT("Couldn't update the WantedPlayRate in CalculateWantedPlayRate, because of invalid PlayRate interval (%f, %f)"), PlayRate.Min, PlayRate.Max);
+				WantedPlayRate = 1.f;
+			}
+			else if (!FMath::IsNearlyEqual(PlayRate.Min, PlayRate.Max, UE_KINDA_SMALL_NUMBER))
+			{
+				TConstArrayView<float> QueryData = SearchContext.GetCachedQuery(CurrentSearchResult.Database->Schema);
+				if (!QueryData.IsEmpty())
+				{
+					if (const UPoseSearchFeatureChannel_Trajectory* TrajectoryChannel = CurrentSearchResult.Database->Schema->FindFirstChannelOfType<UPoseSearchFeatureChannel_Trajectory>())
+					{
+						TConstArrayView<float> ResultData = CurrentSearchResult.Database->GetSearchIndex().GetPoseValues(CurrentSearchResult.PoseIdx);
+						const float EstimatedSpeedRatio = TrajectoryChannel->GetEstimatedSpeedRatio(QueryData, ResultData);
+
+						WantedPlayRate = FMath::Clamp(EstimatedSpeedRatio, PlayRate.Min, PlayRate.Max);
+					}
+					else
+					{
+						UE_LOG(LogPoseSearch, Warning,
+							TEXT("Couldn't update the WantedPlayRate in CalculateWantedPlayRate, because Schema '%s' couldn't find a UPoseSearchFeatureChannel_Trajectory channel"),
+							*GetNameSafe(CurrentSearchResult.Database->Schema));
+					}
+				}
+			}
+			else if (!FMath::IsNearlyZero(TrajectorySpeedMultiplier))
+			{
+				WantedPlayRate = PlayRate.Min / TrajectorySpeedMultiplier;
+			}
+			else
+			{
+				WantedPlayRate = PlayRate.Min;
+			}
+		}
+
+		return WantedPlayRate;
+	}
 }
 
 //////////////////////////////////////////////////////////////////////////
@@ -309,31 +370,34 @@ void FMotionMatchingState::Reset()
 {
 PRAGMA_DISABLE_DEPRECATION_WARNINGS
 	bJumpedToPose = false;
+	WantedPlayRate = 1.f;
 PRAGMA_ENABLE_DEPRECATION_WARNINGS
 
-	CurrentSearchResult.Reset();
+	SearchResult = FPoseSearchBlueprintResult();
 	// Set the elapsed time to INFINITY to trigger a search right away
 	ElapsedPoseSearchTime = std::numeric_limits<float>::infinity();
-	WantedPlayRate = 1.f;
 	PoseIndicesHistory.Reset();
 }
 
 void FMotionMatchingState::AdjustAssetTime(float AssetTime)
 {
+PRAGMA_DISABLE_DEPRECATION_WARNINGS
 	CurrentSearchResult.UpdateWithNormalizedTime(AssetTime);
+PRAGMA_ENABLE_DEPRECATION_WARNINGS
 }
 
 FVector FMotionMatchingState::GetEstimatedFutureRootMotionVelocity() const
 {
 	using namespace UE::PoseSearch;
-	if (CurrentSearchResult.IsValid())
+	if (const UPoseSearchDatabase* Database = SearchResult.SelectedDatabase.Get())
 	{
-		if (const UPoseSearchFeatureChannel_Trajectory* TrajectoryChannel = CurrentSearchResult.Database->Schema->FindFirstChannelOfType<UPoseSearchFeatureChannel_Trajectory>())
+		if (const UPoseSearchFeatureChannel_Trajectory* TrajectoryChannel = Database->Schema->FindFirstChannelOfType<UPoseSearchFeatureChannel_Trajectory>())
 		{
-			const FSearchIndex& SearchIndex = CurrentSearchResult.Database->GetSearchIndex();
+			const int32 PoseIndex = Database->GetPoseIndex(SearchResult.SelectedAnim.Get(), SearchResult.SelectedTime, SearchResult.bIsMirrored, SearchResult.BlendParameters);
+			const FSearchIndex& SearchIndex = Database->GetSearchIndex();
 			if (!SearchIndex.IsValuesEmpty())
 			{
-				TConstArrayView<float> ResultData = SearchIndex.GetPoseValues(CurrentSearchResult.PoseIdx);
+				TConstArrayView<float> ResultData = SearchIndex.GetPoseValues(PoseIndex);
 				return TrajectoryChannel->GetEstimatedFutureRootMotionVelocity(ResultData);
 			}
 		}
@@ -344,59 +408,9 @@ FVector FMotionMatchingState::GetEstimatedFutureRootMotionVelocity() const
 
 void FMotionMatchingState::UpdateWantedPlayRate(const UE::PoseSearch::FSearchContext& SearchContext, const FFloatInterval& PlayRate, float TrajectorySpeedMultiplier, const FPoseSearchEvent& EventToSearch)
 {
-	if (CurrentSearchResult.IsValid())
-	{
-		if (CurrentSearchResult.IsEventSearchResult())
-		{
-			// checking if CurrentSearchResult.EventPoseIdx is part of the EventToSearch.EventTag.
-			// If not, it's an event from a continuing pose search that hasn't been interrupted,
-			// so we keep the previously calculated WantedPlayRate
-			const bool bIsEventSearchFromTag = CurrentSearchResult.IsEventSearchFromTag(EventToSearch.EventTag);
-			if (bIsEventSearchFromTag)
-			{
-				const float TimeToEvent = CurrentSearchResult.CalculateTimeToEvent();
-				if (TimeToEvent > UE_KINDA_SMALL_NUMBER && EventToSearch.TimeToEvent > UE_KINDA_SMALL_NUMBER)
-				{
-					// EventToSearch.TimeToEvent is the desired time to event, and TimeToEvent is the actually current time to event. we calculate WantedPlayRate as ratio between the two
-					WantedPlayRate = TimeToEvent / EventToSearch.TimeToEvent;
-				}
-				// if we passed the event (TimeToEvent <= 0) we leave the WantedPlayRate as previously calculated
-			}
-		}
-		else if (!ensure(PlayRate.Min <= PlayRate.Max && PlayRate.Min > UE_KINDA_SMALL_NUMBER))
-		{
-			UE_LOG(LogPoseSearch, Error, TEXT("Couldn't update the WantedPlayRate in FMotionMatchingState::UpdateWantedPlayRate, because of invalid PlayRate interval (%f, %f)"), PlayRate.Min, PlayRate.Max);
-			WantedPlayRate = 1.f;
-		}
-		else if (!FMath::IsNearlyEqual(PlayRate.Min, PlayRate.Max, UE_KINDA_SMALL_NUMBER))
-		{
-			TConstArrayView<float> QueryData = SearchContext.GetCachedQuery(CurrentSearchResult.Database->Schema);
-			if (!QueryData.IsEmpty())
-			{
-				if (const UPoseSearchFeatureChannel_Trajectory* TrajectoryChannel = CurrentSearchResult.Database->Schema->FindFirstChannelOfType<UPoseSearchFeatureChannel_Trajectory>())
-				{
-					TConstArrayView<float> ResultData = CurrentSearchResult.Database->GetSearchIndex().GetPoseValues(CurrentSearchResult.PoseIdx);
-					const float EstimatedSpeedRatio = TrajectoryChannel->GetEstimatedSpeedRatio(QueryData, ResultData);
-
-					WantedPlayRate = FMath::Clamp(EstimatedSpeedRatio, PlayRate.Min, PlayRate.Max);
-				}
-				else
-				{
-					UE_LOG(LogPoseSearch, Warning,
-						TEXT("Couldn't update the WantedPlayRate in FMotionMatchingState::UpdateWantedPlayRate, because Schema '%s' couldn't find a UPoseSearchFeatureChannel_Trajectory channel"),
-						*GetNameSafe(CurrentSearchResult.Database->Schema));
-				}
-			}
-		}
-		else if (!FMath::IsNearlyZero(TrajectorySpeedMultiplier))
-		{
-			WantedPlayRate = PlayRate.Min / TrajectorySpeedMultiplier;
-		}
-		else
-		{
-			WantedPlayRate = PlayRate.Min;
-		}
-	}
+PRAGMA_DISABLE_DEPRECATION_WARNINGS
+	WantedPlayRate = CalculateWantedPlayRate(CurrentSearchResult, SearchContext, PlayRate, TrajectorySpeedMultiplier, EventToSearch);
+PRAGMA_ENABLE_DEPRECATION_WARNINGS
 }
 
 #if UE_POSE_SEARCH_TRACE_ENABLED
@@ -429,8 +443,6 @@ void UPoseSearchLibrary::TraceMotionMatching(
 
 	TraceState.InterruptMode = InterruptMode;
 
-	const FSearchResult& CurrentResult = SearchResult;
-
 	const int32 AnimContextsNum = SearchContext.GetAnimContexts().Num();
 	TraceState.SkeletalMeshComponentIds.SetNum(AnimContextsNum);
 
@@ -484,7 +496,7 @@ void UPoseSearchLibrary::TraceMotionMatching(
 
 	TArray<uint64, TInlineAllocator<64>> DatabaseIds;
 	int32 DbEntryIdx = 0;
-	const int32 CurrentPoseIdx = bSearch && CurrentResult.PoseCost.IsValid() ? CurrentResult.PoseIdx : INDEX_NONE;
+	const int32 CurrentPoseIdx = bSearch && SearchResult.PoseCost.IsValid() ? SearchResult.PoseIdx : INDEX_NONE;
 	TraceState.DatabaseEntries.SetNum(SearchContext.GetBestPoseCandidatesMap().Num());
 	for (TPair<const UPoseSearchDatabase*, FSearchContext::FBestPoseCandidates> DatabaseBestPoseCandidates : SearchContext.GetBestPoseCandidatesMap())
 	{
@@ -506,7 +518,7 @@ void UPoseSearchLibrary::TraceMotionMatching(
 			PoseEntry.DbPoseIdx = PoseCandidate.PoseIdx;
 			PoseEntry.Cost = PoseCandidate.Cost;
 			PoseEntry.PoseCandidateFlags = PoseCandidate.PoseCandidateFlags;
-			if (CurrentPoseIdx == PoseCandidate.PoseIdx && CurrentResult.Database.Get() == Database)
+			if (CurrentPoseIdx == PoseCandidate.PoseIdx && SearchResult.Database.Get() == Database)
 			{
 				check(EnumHasAnyFlags(PoseEntry.PoseCandidateFlags, EPoseCandidateFlags::Valid_Pose | EPoseCandidateFlags::Valid_ContinuingPose));
 
@@ -549,8 +561,8 @@ void UPoseSearchLibrary::TraceMotionMatching(
 			}
 		}
 
-		const FSearchIndexAsset* SearchIndexAsset = CurrentResult.GetSearchIndexAsset();
-		const UPoseSearchDatabase* CurrentResultDatabase = CurrentResult.Database.Get();
+		const FSearchIndexAsset* SearchIndexAsset = SearchResult.GetSearchIndexAsset();
+		const UPoseSearchDatabase* CurrentResultDatabase = SearchResult.Database.Get();
 		if (SearchIndexAsset && CurrentResultDatabase)
 		{
 			const FPoseSearchDatabaseAnimationAssetBase* DatabaseAsset = CurrentResultDatabase->GetDatabaseAnimationAsset<FPoseSearchDatabaseAnimationAssetBase>(*SearchIndexAsset);
@@ -561,14 +573,14 @@ void UPoseSearchLibrary::TraceMotionMatching(
 				FAnimationAssetSampler Sampler(AnimationAsset, FTransform::Identity,FVector::ZeroVector, FAnimationAssetSampler::DefaultRootTransformSamplingRate, true, false);
 
 				const float TimeStep = DeltaTime * WantedPlayRate;
-				const FTransform PrevRoot = Sampler.ExtractRootTransform(CurrentResult.AssetTime);
-				const FTransform CurrRoot = Sampler.ExtractRootTransform(CurrentResult.AssetTime + TimeStep);
+				const FTransform PrevRoot = Sampler.ExtractRootTransform(SearchResult.AssetTime);
+				const FTransform CurrRoot = Sampler.ExtractRootTransform(SearchResult.AssetTime + TimeStep);
 				const FTransform RootMotionTransformDelta = PrevRoot.GetRelativeTransform(CurrRoot);
 				TraceState.AnimLinearVelocity = RootMotionTransformDelta.GetTranslation().Size() / DeltaTime;
 				TraceState.AnimAngularVelocity = FMath::RadiansToDegrees(RootMotionTransformDelta.GetRotation().GetAngle()) / DeltaTime;
 
 				// Need another root motion extraction for non-playrate version in case acceleration isn't the same.
-				const FTransform CurrRootNoTimescale = Sampler.ExtractRootTransform(CurrentResult.AssetTime + DeltaTime);
+				const FTransform CurrRootNoTimescale = Sampler.ExtractRootTransform(SearchResult.AssetTime + DeltaTime);
 				const FTransform RootMotionTransformDeltaNoTimescale = PrevRoot.GetRelativeTransform(CurrRootNoTimescale);
 				TraceState.AnimLinearVelocityNoTimescale = RootMotionTransformDeltaNoTimescale.GetTranslation().Size() / DeltaTime;
 				TraceState.AnimAngularVelocityNoTimescale = FMath::RadiansToDegrees(RootMotionTransformDeltaNoTimescale.GetRotation().GetAngle()) / DeltaTime;
@@ -578,14 +590,14 @@ void UPoseSearchLibrary::TraceMotionMatching(
 	}
 
 	TraceState.ElapsedPoseSearchTime = ElapsedPoseSearchTime;
-	TraceState.AssetPlayerTime = CurrentResult.AssetTime;
+	TraceState.AssetPlayerTime = SearchResult.AssetTime;
 	TraceState.DeltaTime = DeltaTime;
 
 	TraceState.RecordingTime = RecordingTime;
-	TraceState.SearchBestCost = CurrentResult.PoseCost;
+	TraceState.SearchBestCost = SearchResult.PoseCost;
 #if WITH_EDITOR && ENABLE_ANIM_DEBUG
-	TraceState.SearchBruteForceCost = CurrentResult.BruteForcePoseCost;
-	TraceState.SearchBestPosePos = CurrentResult.BestPosePos;
+	TraceState.SearchBruteForceCost = SearchResult.BruteForcePoseCost;
+	TraceState.SearchBestPosePos = SearchResult.BestPosePos;
 #else // WITH_EDITOR && ENABLE_ANIM_DEBUG
 	TraceState.SearchBruteForceCost = 0.f;
 	TraceState.SearchBestPosePos = 0;
@@ -663,17 +675,45 @@ void UPoseSearchLibrary::UpdateMotionMatchingState(
 
 	FMemMark Mark(FMemStack::Get());
 
-	const UPoseSearchDatabase* CurrentResultDatabase = InOutMotionMatchingState.CurrentSearchResult.Database.Get();
+	UE::PoseSearch::FSearchResult InternalSearchResult;
+
+	const UPoseSearchDatabase* CurrentResultDatabase = InOutMotionMatchingState.SearchResult.SelectedDatabase.Get();
 	if (IsInvalidatingContinuingPose(InterruptMode, CurrentResultDatabase, Databases))
 	{
-		InOutMotionMatchingState.CurrentSearchResult.Reset();
+		InOutMotionMatchingState.SearchResult = FPoseSearchBlueprintResult();
+	}
+	else
+	{
+		InternalSearchResult.InitFrom(InOutMotionMatchingState.SearchResult);
+
+#if DO_CHECK
+		if (InternalSearchResult.PoseIdx != INDEX_NONE)
+		{
+			FPoseSearchBlueprintResult TestSearchResult;
+			TestSearchResult.InitFrom(InternalSearchResult, InOutMotionMatchingState.SearchResult.WantedPlayRate);
+
+			if (InOutMotionMatchingState.SearchResult.SelectedAnim != TestSearchResult.SelectedAnim ||
+				InOutMotionMatchingState.SearchResult.SelectedTime != TestSearchResult.SelectedTime ||
+				InOutMotionMatchingState.SearchResult.bIsContinuingPoseSearch != TestSearchResult.bIsContinuingPoseSearch ||
+				InOutMotionMatchingState.SearchResult.WantedPlayRate != TestSearchResult.WantedPlayRate ||
+				InOutMotionMatchingState.SearchResult.bLoop != TestSearchResult.bLoop ||
+				InOutMotionMatchingState.SearchResult.bIsMirrored != TestSearchResult.bIsMirrored ||
+				InOutMotionMatchingState.SearchResult.BlendParameters != TestSearchResult.BlendParameters ||
+				InOutMotionMatchingState.SearchResult.SelectedDatabase != TestSearchResult.SelectedDatabase ||
+				InOutMotionMatchingState.SearchResult.SearchCost != TestSearchResult.SearchCost ||
+				InOutMotionMatchingState.SearchResult.bIsInteraction != TestSearchResult.bIsInteraction)
+			{
+				UE_LOG(LogPoseSearch, Error, TEXT("Error converting FPoseSearchBlueprintResult to UE::PoseSearch::FSearchResult!"));
+			}
+		}
+#endif // DO_CHECK
 	}
 
 	const FPoseSearchEvent PlayRateOverriddenEvent = EventToSearch.GetPlayRateOverriddenEvent(PlayRate);
-	FSearchContext SearchContext(0.f, &InOutMotionMatchingState.PoseIndicesHistory, InOutMotionMatchingState.CurrentSearchResult, PoseJumpThresholdTime, PlayRateOverriddenEvent);
+	FSearchContext SearchContext(0.f, &InOutMotionMatchingState.PoseIndicesHistory, InternalSearchResult, PoseJumpThresholdTime, PlayRateOverriddenEvent);
 	SearchContext.AddRole(GetCommonDefaultRole(Databases), AnimContext, PoseHistory);
 
-	const bool bCanAdvance = InOutMotionMatchingState.CurrentSearchResult.CanAdvance(DeltaTime);
+	const bool bCanAdvance = InternalSearchResult.PoseIdx != INDEX_NONE;
 
 	// If we can't advance or enough time has elapsed since the last pose jump then search
 	const bool bSearch = !bCanAdvance || (InOutMotionMatchingState.ElapsedPoseSearchTime >= SearchThrottleTime);
@@ -720,7 +760,7 @@ void UPoseSearchLibrary::UpdateMotionMatchingState(
 		// resetting CurrentSearchResult if any DDC indexing on the requested databases is still in progress
 		if (SearchContext.IsAsyncBuildIndexInProgress())
 		{
-			InOutMotionMatchingState.CurrentSearchResult.Reset();
+			InternalSearchResult.Reset();
 		}
 #endif // WITH_EDITOR
 
@@ -767,26 +807,32 @@ void UPoseSearchLibrary::UpdateMotionMatchingState(
 #endif // !NO_LOGGING
 
 		// Remember which pose and sequence we're playing from the database
-		InOutMotionMatchingState.CurrentSearchResult = SearchResult;
+		InternalSearchResult = SearchResult;
 	}
 	else
 	{
 		// @todo: for blendspaces ElapsedPoseSearchTime should be incremented by a normalized DeltaTime, also didn't we already synchronized the result???
 		InOutMotionMatchingState.ElapsedPoseSearchTime += DeltaTime;
-		InOutMotionMatchingState.CurrentSearchResult.bIsContinuingPoseSearch = true;
+		InternalSearchResult.bIsContinuingPoseSearch = true;
 	}
 
 	// @todo: consider moving this into if (bSearch) to avoid calling SearchContext.GetCachedQuery if no search is required
-	InOutMotionMatchingState.UpdateWantedPlayRate(SearchContext, PlayRate, PoseHistory ? PoseHistory->GetTrajectorySpeedMultiplier() : 1.f, EventToSearch);
+	const float WantedPlayRate = CalculateWantedPlayRate(InternalSearchResult, SearchContext, PlayRate, PoseHistory ? PoseHistory->GetTrajectorySpeedMultiplier() : 1.f, EventToSearch);
 
-	InOutMotionMatchingState.PoseIndicesHistory.Update(InOutMotionMatchingState.CurrentSearchResult, DeltaTime, PoseReselectHistory);
+	InOutMotionMatchingState.PoseIndicesHistory.Update(InternalSearchResult, DeltaTime, PoseReselectHistory);
+	InOutMotionMatchingState.SearchResult.InitFrom(InternalSearchResult, WantedPlayRate);
+
+PRAGMA_DISABLE_DEPRECATION_WARNINGS
+	InOutMotionMatchingState.WantedPlayRate = WantedPlayRate;
+	InOutMotionMatchingState.CurrentSearchResult = InternalSearchResult;
+PRAGMA_ENABLE_DEPRECATION_WARNINGS
 
 #if UE_POSE_SEARCH_TRACE_ENABLED
-	TraceMotionMatching(SearchContext, InOutMotionMatchingState.CurrentSearchResult, InOutMotionMatchingState.ElapsedPoseSearchTime, DeltaTime, bSearch, InOutMotionMatchingState.WantedPlayRate, InterruptMode);
+	TraceMotionMatching(SearchContext, InternalSearchResult, InOutMotionMatchingState.ElapsedPoseSearchTime, DeltaTime, bSearch, InOutMotionMatchingState.SearchResult.WantedPlayRate, InterruptMode);
 #endif // UE_POSE_SEARCH_TRACE_ENABLED
 
 #if WITH_EDITORONLY_DATA && ENABLE_ANIM_DEBUG
-	const FSearchResult& CurResult = InOutMotionMatchingState.CurrentSearchResult;
+	const FSearchResult& CurResult = InternalSearchResult;
 	if (bDebugDrawQuery || bDebugDrawCurResult)
 	{
 		const UPoseSearchDatabase* CurResultDatabase = CurResult.Database.Get();
@@ -925,34 +971,25 @@ void UPoseSearchLibrary::MotionMatch(
 	if (SearchResult.IsValid())
 	{
 		const UPoseSearchDatabase* Database = SearchResult.Database.Get();
-		const FSearchIndexAsset* SearchIndexAsset = SearchResult.GetSearchIndexAsset();
-		if (const FPoseSearchDatabaseAnimationAssetBase* DatabaseAsset = Database->GetDatabaseAnimationAsset<FPoseSearchDatabaseAnimationAssetBase>(*SearchIndexAsset))
+		check(Database);
+		
+		// figuring out the WantedPlayRate
+		float WantedPlayRate = 1.f;
+		if (Future.Animation && Future.IntervalTime > 0.f)
 		{
-			Result.SelectedAnim = DatabaseAsset->GetAnimationAsset();
-			Result.SelectedTime = SearchResult.AssetTime;
-			Result.bIsContinuingPoseSearch = SearchResult.bIsContinuingPoseSearch;
-			Result.bLoop = SearchIndexAsset->IsLooping();
-			Result.bIsMirrored = SearchIndexAsset->IsMirrored();
-			Result.BlendParameters = SearchIndexAsset->GetBlendParameters();
-			Result.SelectedDatabase = Database;
-			Result.SearchCost = SearchResult.PoseCost;
-			
-			// figuring out the WantedPlayRate
-			Result.WantedPlayRate = 1.f;
-			if (Future.Animation && Future.IntervalTime > 0.f)
+			if (const UPoseSearchFeatureChannel_PermutationTime* PermutationTimeChannel = Database->Schema->FindFirstChannelOfType<UPoseSearchFeatureChannel_PermutationTime>())
 			{
-				if (const UPoseSearchFeatureChannel_PermutationTime* PermutationTimeChannel = Database->Schema->FindFirstChannelOfType<UPoseSearchFeatureChannel_PermutationTime>())
+				const FSearchIndex& SearchIndex = Database->GetSearchIndex();
+				if (!SearchIndex.IsValuesEmpty())
 				{
-					const FSearchIndex& SearchIndex = Database->GetSearchIndex();
-					if (!SearchIndex.IsValuesEmpty())
-					{
-						TConstArrayView<float> ResultData = Database->GetSearchIndex().GetPoseValues(SearchResult.PoseIdx);
-						const float ActualIntervalTime = PermutationTimeChannel->GetPermutationTime(ResultData);
-						Result.WantedPlayRate = ActualIntervalTime / Future.IntervalTime;
-					}
+					TConstArrayView<float> ResultData = Database->GetSearchIndex().GetPoseValues(SearchResult.PoseIdx);
+					const float ActualIntervalTime = PermutationTimeChannel->GetPermutationTime(ResultData);
+					WantedPlayRate = ActualIntervalTime / Future.IntervalTime;
 				}
 			}
 		}
+
+		Result.InitFrom(SearchResult, WantedPlayRate);
 	}
 }
 
@@ -1101,102 +1138,59 @@ UE::PoseSearch::FSearchResult UPoseSearchLibrary::MotionMatch(
 		const bool bInvalidatingContinuingPose = IsInvalidatingContinuingPose(ContinuingProperties.InterruptMode, Database, AssetsToSearchPerDatabaseMap);
 		if (!bInvalidatingContinuingPose)
 		{
-			const bool bForceInterrupt = IsForceInterrupt(ContinuingProperties.InterruptMode, Database, AssetsToSearchPerDatabaseMap);
+			ReconstructedPreviousSearchResult.AssetTime = ContinuingProperties.PlayingAssetAccumulatedTime;
+			ReconstructedPreviousSearchResult.PoseIdx = Database->GetPoseIndex(ContinuingProperties.PlayingAsset.Get(), ContinuingProperties.PlayingAssetAccumulatedTime, ContinuingProperties.bIsPlayingAssetMirrored, ContinuingProperties.PlayingAssetBlendParameters);
+			ReconstructedPreviousSearchResult.Database = Database;
 
-			const FSearchIndex& SearchIndex = Database->GetSearchIndex();
-			for (int32 AssetIndex : Database->GetAssetIndexesForSourceAsset(ContinuingProperties.PlayingAsset.Get()))
+			const bool bForceInterrupt = IsForceInterrupt(ContinuingProperties.InterruptMode, Database, AssetsToSearchPerDatabaseMap);
+			const bool bCanAdvance = ReconstructedPreviousSearchResult.PoseIdx != INDEX_NONE;
+			if (bCanAdvance && !bForceInterrupt)
 			{
-				const FSearchIndexAsset& SearchIndexAsset = SearchIndex.Assets[AssetIndex];
+				SearchContext.UpdateCurrentResultPoseVector();
+				const FSearchResult NewSearchResult = Database->SearchContinuingPose(SearchContext);
 
-				const float FirstSampleTime = SearchIndexAsset.GetFirstSampleTime(Database->Schema->SampleRate);
-				const float LastSampleTime = SearchIndexAsset.GetLastSampleTime(Database->Schema->SampleRate);
+#if WITH_EDITOR && ENABLE_ANIM_DEBUG && UE_POSE_SEARCH_TRACE_ENABLED
+				const FPoseSearchCost BestBruteForcePoseCost = NewSearchResult.BruteForcePoseCost < SearchResult.BruteForcePoseCost ? NewSearchResult.BruteForcePoseCost : SearchResult.BruteForcePoseCost;
+#endif // WITH_EDITOR && ENABLE_ANIM_DEBUG && UE_POSE_SEARCH_TRACE_ENABLED
 
-				bool bCanAdvance = true;
-				if (SearchIndexAsset.IsLooping())
-				{
-					const float DeltaSampleTime = LastSampleTime - FirstSampleTime;
-					if (DeltaSampleTime < UE_SMALL_NUMBER)
-					{
-						ReconstructedPreviousSearchResult.AssetTime = FirstSampleTime;
-					}
-					else if (ContinuingProperties.PlayingAssetAccumulatedTime < FirstSampleTime)
-					{
-						ReconstructedPreviousSearchResult.AssetTime = FMath::Fmod(ContinuingProperties.PlayingAssetAccumulatedTime - FirstSampleTime, DeltaSampleTime) + DeltaSampleTime + FirstSampleTime;
-					}
-					else if (ContinuingProperties.PlayingAssetAccumulatedTime > LastSampleTime)
-					{
-						ReconstructedPreviousSearchResult.AssetTime = FMath::Fmod(ContinuingProperties.PlayingAssetAccumulatedTime - FirstSampleTime, DeltaSampleTime) + FirstSampleTime;
-					}
-					else
-					{
-						ReconstructedPreviousSearchResult.AssetTime = ContinuingProperties.PlayingAssetAccumulatedTime;
-					}
-				}
-				else
+				if (NewSearchResult.PoseCost < SearchResult.PoseCost)
 				{
-					const float MaxTimeToBeAbleToContinuingPlayingAnimation = LastSampleTime - DeltaSeconds;
-					bCanAdvance = ContinuingProperties.PlayingAssetAccumulatedTime >= FirstSampleTime && ContinuingProperties.PlayingAssetAccumulatedTime < MaxTimeToBeAbleToContinuingPlayingAnimation;
-					ReconstructedPreviousSearchResult.AssetTime = ContinuingProperties.PlayingAssetAccumulatedTime;
+					SearchResult = NewSearchResult;
+					SearchContext.UpdateCurrentBestCost(SearchResult.PoseCost);
 				}
 
-				ReconstructedPreviousSearchResult.Database = Database;
-				ReconstructedPreviousSearchResult.PoseIdx = Database->GetPoseIndexFromTime(ContinuingProperties.PlayingAssetAccumulatedTime, SearchIndexAsset);
-
-				if (bCanAdvance && !bForceInterrupt)
-				{
-					SearchContext.UpdateCurrentResultPoseVector();
-					const FSearchResult NewSearchResult = Database->SearchContinuingPose(SearchContext);
-
-#if WITH_EDITOR && ENABLE_ANIM_DEBUG && UE_POSE_SEARCH_TRACE_ENABLED
-					const FPoseSearchCost BestBruteForcePoseCost = NewSearchResult.BruteForcePoseCost < SearchResult.BruteForcePoseCost ? NewSearchResult.BruteForcePoseCost : SearchResult.BruteForcePoseCost;
-#endif // WITH_EDITOR && ENABLE_ANIM_DEBUG && UE_POSE_SEARCH_TRACE_ENABLED
-
-					if (NewSearchResult.PoseCost < SearchResult.PoseCost)
-					{
-						SearchResult = NewSearchResult;
-						SearchContext.UpdateCurrentBestCost(SearchResult.PoseCost);
-					}
 #if WITH_EDITOR && ENABLE_ANIM_DEBUG && UE_POSE_SEARCH_TRACE_ENABLED
-					SearchResult.BruteForcePoseCost = BestBruteForcePoseCost;
+				SearchResult.BruteForcePoseCost = BestBruteForcePoseCost;
 #endif // WITH_EDITOR && ENABLE_ANIM_DEBUG && UE_POSE_SEARCH_TRACE_ENABLED
 
-					if (NewSearchResult.IsValid())
+				if (NewSearchResult.IsValid())
+				{
+					// keeping track of the best ReconstructedPreviousSearchResult, as best continuing pose search result for the full Database search happening after
+					if (FSearchResult* FoundReconstructedPreviousSearchBestResult = ReconstructedPreviousSearchBestResultMap.Find(Database))
 					{
-						// keeping track of the best ReconstructedPreviousSearchResult, as best continuing pose search result for the full Database search happening after
-						if (FSearchResult* FoundReconstructedPreviousSearchBestResult = ReconstructedPreviousSearchBestResultMap.Find(Database))
+						if (NewSearchResult.PoseCost == FoundReconstructedPreviousSearchBestResult->PoseCost)
 						{
-							if (NewSearchResult.PoseCost == FoundReconstructedPreviousSearchBestResult->PoseCost)
-							{
-								check(NewSearchResult.PoseIdx != FoundReconstructedPreviousSearchBestResult->PoseIdx);
-								if (NewSearchResult.PoseIdx < FoundReconstructedPreviousSearchBestResult->PoseIdx)
-								{
-									// choosing to update with the lowest PoseIdx to avoid indeterminism 
-									// since ReconstructedPreviousSearchBestResultMap is an unsorted Map
-									*FoundReconstructedPreviousSearchBestResult = ReconstructedPreviousSearchResult;
-									FoundReconstructedPreviousSearchBestResult->PoseCost = NewSearchResult.PoseCost;
-								}
-							}
-							else if (NewSearchResult.PoseCost < FoundReconstructedPreviousSearchBestResult->PoseCost)
+							check(NewSearchResult.PoseIdx != FoundReconstructedPreviousSearchBestResult->PoseIdx);
+							if (NewSearchResult.PoseIdx < FoundReconstructedPreviousSearchBestResult->PoseIdx)
 							{
+								// choosing to update with the lowest PoseIdx to avoid indeterminism 
+								// since ReconstructedPreviousSearchBestResultMap is an unsorted Map
 								*FoundReconstructedPreviousSearchBestResult = ReconstructedPreviousSearchResult;
 								FoundReconstructedPreviousSearchBestResult->PoseCost = NewSearchResult.PoseCost;
 							}
 						}
-						else
+						else if (NewSearchResult.PoseCost < FoundReconstructedPreviousSearchBestResult->PoseCost)
 						{
-							FSearchResult& NewReconstructedPreviousSearchBestResult = ReconstructedPreviousSearchBestResultMap.Add(Database);
-							NewReconstructedPreviousSearchBestResult = ReconstructedPreviousSearchResult;
-							NewReconstructedPreviousSearchBestResult.PoseCost = NewSearchResult.PoseCost;
-
-							if (bForceInterrupt)
-							{
-								// no point to continue if it's bForceInterrupt: no continuing pose searches will be performed 
-								// so all the ReconstructedPreviousSearchResult will lead to the same pose with invalid PoseCost
-								// so any of the pair AssetTime / PoseIdx will be the same
-								break;
-							}
+							*FoundReconstructedPreviousSearchBestResult = ReconstructedPreviousSearchResult;
+							FoundReconstructedPreviousSearchBestResult->PoseCost = NewSearchResult.PoseCost;
 						}
 					}
+					else
+					{
+						FSearchResult& NewReconstructedPreviousSearchBestResult = ReconstructedPreviousSearchBestResultMap.Add(Database);
+						NewReconstructedPreviousSearchBestResult = ReconstructedPreviousSearchResult;
+						NewReconstructedPreviousSearchBestResult.PoseCost = NewSearchResult.PoseCost;
+					}
 				}
 			}
 		}
diff --git a/Engine/Plugins/Animation/PoseSearch/Source/Runtime/Private/PoseSearchModule.cpp b/Engine/Plugins/Animation/PoseSearch/Source/Runtime/Private/PoseSearchModule.cpp
--- a/Engine/Plugins/Animation/PoseSearch/Source/Runtime/Private/PoseSearchModule.cpp
+++ b/Engine/Plugins/Animation/PoseSearch/Source/Runtime/Private/PoseSearchModule.cpp
@@ -57,6 +57,8 @@ public:
 		FPoseSearchContinuingProperties ContinuingProperties;
 		ContinuingProperties.PlayingAsset = PlayingAsset.Asset;
 		ContinuingProperties.PlayingAssetAccumulatedTime = PlayingAsset.AccumulatedTime;
+		ContinuingProperties.bIsPlayingAssetMirrored = PlayingAsset.bMirrored;
+		ContinuingProperties.PlayingAssetBlendParameters = PlayingAsset.BlendParameters;
 
 		FPoseSearchFutureProperties Future;
 		Future.Animation = FutureAsset.Asset;
@@ -85,6 +87,7 @@ public:
 				ProviderResult.TimeOffsetSeconds = SearchResult.AssetTime;
 				ProviderResult.bIsFromContinuingPlaying = SearchResult.bIsContinuingPoseSearch;
 				ProviderResult.bMirrored = SearchIndexAsset->IsMirrored();
+				ProviderResult.BlendParameters = SearchIndexAsset->GetBlendParameters();
 
 				// figuring out the WantedPlayRate
 				ProviderResult.WantedPlayRate = 1.f;
diff --git a/Engine/Plugins/Animation/PoseSearch/Source/Runtime/Private/PoseSearchResult.cpp b/Engine/Plugins/Animation/PoseSearch/Source/Runtime/Private/PoseSearchResult.cpp
--- a/Engine/Plugins/Animation/PoseSearch/Source/Runtime/Private/PoseSearchResult.cpp
+++ b/Engine/Plugins/Animation/PoseSearch/Source/Runtime/Private/PoseSearchResult.cpp
@@ -2,15 +2,48 @@
 
 #include "PoseSearch/PoseSearchResult.h"
 #include "Animation/BlendSpace.h"
-#include "StructUtils/InstancedStruct.h"
+#include "PoseSearch/MultiAnimAsset.h"
 #include "PoseSearch/PoseSearchDatabase.h"
 #include "PoseSearch/PoseSearchSchema.h"
+#include "StructUtils/InstancedStruct.h"
 
 namespace UE::PoseSearch
 {
 
+bool FSearchResult::DebugValidate() const
+{
+	bool bIsValidated = true;
+
+#if WITH_EDITOR && ENABLE_ANIM_DEBUG
+
+	if (IsValid())
+	{
+		const FSearchIndex& SearchIndex = Database->GetSearchIndex();
+		const FSearchIndexAsset& SearchIndexAsset = SearchIndex.GetAssetForPose(PoseIdx);
+		if (const FPoseSearchDatabaseAnimationAssetBase* DatabaseAnimationAssetBase = Database->GetDatabaseAnimationAsset<FPoseSearchDatabaseAnimationAssetBase>(SearchIndexAsset))
+		{
+			const float RealAssetTime = AssetTime * SearchIndexAsset.GetToRealTimeFactor();
+			const int32 RecalculatedPoseIdx = SearchIndexAsset.GetPoseIndexFromTime(RealAssetTime, Database->Schema->SampleRate);
+			if (RecalculatedPoseIdx != PoseIdx)
+			{
+				bIsValidated = false;
+			}
+		}
+		else
+		{
+			bIsValidated = false;
+		}
+	}
+
+#endif // WITH_EDITOR && ENABLE_ANIM_DEBUG
+
+	return bIsValidated;
+}
+
 void FSearchResult::UpdateWithNormalizedTime(float NormalizedTime)
 {
+	check(DebugValidate());
+
 	if (IsValid())
 	{
 		// for non blend spaces the real time corrisponds to the normalized time!
@@ -41,17 +74,21 @@ void FSearchResult::UpdateWithNormalizedTime(float NormalizedTime)
 		}
 #endif // WITH_EDITOR && DO_CHECK
 		
-		PoseIdx = Database->GetPoseIndexFromTime(RealTime, SearchIndexAsset);
+		PoseIdx = SearchIndexAsset.GetPoseIndexFromTime(RealTime, Database->Schema->SampleRate);
 		AssetTime = NormalizedTime;
 	}
 	else
 	{
 		Reset();
 	}
+	
+	check(DebugValidate());
 }
 
 void FSearchResult::UpdateWithRealTime(float RealTime)
 {
+	check(DebugValidate());
+
 	if (IsValid())
 	{
 		const FSearchIndexAsset& SearchIndexAsset = Database->GetSearchIndex().GetAssetForPose(PoseIdx);
@@ -82,13 +119,15 @@ void FSearchResult::UpdateWithRealTime(float RealTime)
 		}
 #endif // WITH_EDITOR && DO_CHECK
 
-		PoseIdx = Database->GetPoseIndexFromTime(RealTime, SearchIndexAsset);
+		PoseIdx = SearchIndexAsset.GetPoseIndexFromTime(RealTime, Database->Schema->SampleRate);
 		AssetTime = NormalizedTime;
 	}
 	else
 	{
 		Reset();
 	}
+	
+	check(DebugValidate());
 }
 
 bool FSearchResult::IsEventSearchFromTag(const FGameplayTag& EventTag) const
@@ -151,26 +190,71 @@ const UAnimationAsset* FSearchResult::GetCurrentResultAnimationAsset(const FRole
 	return nullptr;
 }
 
-bool FSearchResult::CanAdvance(float DeltaTime) const
+void FSearchResult::InitFrom(const FPoseSearchBlueprintResult& BlueprintResult)
+{
+	PoseCost = FPoseSearchCost(BlueprintResult.SearchCost, 0.f, 0.f);
+	PoseIdx = BlueprintResult.SelectedDatabase.IsValid() ? BlueprintResult.SelectedDatabase->GetPoseIndex(BlueprintResult.SelectedAnim.Get(), BlueprintResult.SelectedTime, BlueprintResult.bIsMirrored, BlueprintResult.BlendParameters) : INDEX_NONE;
+	EventPoseIdx = INDEX_NONE;
+	Database = BlueprintResult.SelectedDatabase;
+	AssetTime = BlueprintResult.SelectedTime;
+	bIsContinuingPoseSearch = BlueprintResult.bIsContinuingPoseSearch;
+
+#if WITH_EDITOR && ENABLE_ANIM_DEBUG && UE_POSE_SEARCH_TRACE_ENABLED
+	BruteForcePoseCost = FPoseSearchCost();
+	BestPosePos = 0;
+#endif // WITH_EDITOR && ENABLE_ANIM_DEBUG && UE_POSE_SEARCH_TRACE_ENABLED
+}
+
+} // namespace UE::PoseSearch
+
+bool FPoseSearchBlueprintResult::InitFrom(const UE::PoseSearch::FSearchResult& SearchResult, float InWantedPlayRate)
 {
-	const FSearchIndexAsset* SearchIndexAsset = GetSearchIndexAsset();
-	if (!SearchIndexAsset)
+	using namespace UE::PoseSearch;
+
+	if (const FSearchIndexAsset* SearchIndexAsset = SearchResult.GetSearchIndexAsset())
 	{
-		return false;
+		const UPoseSearchDatabase* Database = SearchResult.Database.Get();
+
+		const FPoseSearchDatabaseAnimationAssetBase* DatabaseAsset = Database->GetDatabaseAnimationAsset<FPoseSearchDatabaseAnimationAssetBase>(*SearchIndexAsset);
+		check(DatabaseAsset);
+		
+#if WITH_EDITORONLY_DATA
+		SelectedAnimation_DEPRECATED = DatabaseAsset->GetAnimationAsset();
+#endif // WITH_EDITORONLY_DATA
+
+		SelectedAnim = DatabaseAsset->GetAnimationAsset();
+		SelectedTime = SearchResult.AssetTime;
+		bIsContinuingPoseSearch = SearchResult.bIsContinuingPoseSearch;
+		WantedPlayRate = InWantedPlayRate;
+		bLoop = SearchIndexAsset->IsLooping();
+		bIsMirrored = SearchIndexAsset->IsMirrored();
+		BlendParameters = SearchIndexAsset->GetBlendParameters();
+		SelectedDatabase = Database;
+		SearchCost = SearchResult.PoseCost;
+		Role = Database->Schema->GetDefaultRole();
+
+		bIsInteraction = false;
+		ActorRootTransforms.Reset();
+		ActorRootBoneTransforms.Reset();
+		AnimContexts.Reset();
+		return true;
 	}
 	
-	if (SearchIndexAsset->IsLooping())
+	*this = FPoseSearchBlueprintResult();
+	return false;
+}
+
+UAnimationAsset* FPoseSearchBlueprintResult::GetAnimationAssetForRole()
+{
+	if (UAnimationAsset* AnimationAsset = Cast<UAnimationAsset>(SelectedAnim))
 	{
-		return true;
+		return AnimationAsset;
 	}
-
-	// leaving this if statement for debuggability
-	if (INDEX_NONE != Database->GetPoseIndexFromTime(AssetTime + DeltaTime, *SearchIndexAsset))
+	
+	if (UMultiAnimAsset* MultiAnimAsset = Cast<UMultiAnimAsset>(SelectedAnim))
 	{
-		return true;
+		return MultiAnimAsset->GetAnimationAsset(Role);
 	}
 
-	return false;
-}
-
-} // namespace UE::PoseSearch
+	return nullptr;
+}
\ No newline at end of file
diff --git a/Engine/Plugins/Animation/PoseSearch/Source/Runtime/Public/PoseSearch/PoseSearchDatabase.h b/Engine/Plugins/Animation/PoseSearch/Source/Runtime/Public/PoseSearch/PoseSearchDatabase.h
--- a/Engine/Plugins/Animation/PoseSearch/Source/Runtime/Public/PoseSearch/PoseSearchDatabase.h
+++ b/Engine/Plugins/Animation/PoseSearch/Source/Runtime/Public/PoseSearch/PoseSearchDatabase.h
@@ -539,7 +539,15 @@ public:
 	
 	bool GetSkipSearchIfPossible() const;
 
-	int32 GetPoseIndexFromTime(float AssetTime, const UE::PoseSearch::FSearchIndexAsset& SearchIndexAsset) const;
+	UE_DEPRECATED(5.6, "Use SearchIndexAsset.GetPoseIndexFromTime instead")
+	// given RealTimeInSeconds in seconds (not a normalized time, so for blend spaces needs to be multiplied by the 
+	// blendspace length given the blend parameters). it'll return the associated pose index
+	int32 GetPoseIndexFromTime(float RealTimeInSeconds, const UE::PoseSearch::FSearchIndexAsset& SearchIndexAsset) const;
+
+	// given an AnimationAsset, a normalized time AnimationAssetTime (in [0-1] range for blend spaces), and mirror state bMirrored
+	// will retrieve the associated pose index from this database closest in distance to BlendParameters.
+	// It'll return INDEX_NONE if it cannot find a pose
+	int32 GetPoseIndex(const UObject* AnimationAsset, float AnimationAssetTime, bool bMirrored, const FVector& BlendParameters) const;
 
 	void AddAnimationAsset(FInstancedStruct AnimationAsset);
 	void RemoveAnimationAssetAt(int32 AnimationAssetIndex);
diff --git a/Engine/Plugins/Animation/PoseSearch/Source/Runtime/Public/PoseSearch/PoseSearchIndex.h b/Engine/Plugins/Animation/PoseSearch/Source/Runtime/Public/PoseSearch/PoseSearchIndex.h
--- a/Engine/Plugins/Animation/PoseSearch/Source/Runtime/Public/PoseSearch/PoseSearchIndex.h
+++ b/Engine/Plugins/Animation/PoseSearch/Source/Runtime/Public/PoseSearch/PoseSearchIndex.h
@@ -97,7 +97,7 @@ struct FSearchIndexAsset
 		int32 InFirstPoseIdx,
 		int32 InFirstSampleIdx,
 		int32 InLastSampleIdx,
-		float InToRealTimeFactor /*= 1.f*/)
+		float InToRealTimeFactor = 1.f)
 		: SourceAssetIdx(InSourceAssetIdx)
 		, bMirrored(bInMirrored)
 		, bLooping(bInLooping)
@@ -123,8 +123,8 @@ struct FSearchIndexAsset
 		const FFloatInterval& InSamplingInterval,
 		int32 SchemaSampleRate,
 		int32 InPermutationIdx,
-		FVector InBlendParameters /*= FVector::Zero()*/,
-		float InToRealTimeFactor /*= 1.f*/)
+		FVector InBlendParameters = FVector::Zero(),
+		float InToRealTimeFactor = 1.f)
 		: FSearchIndexAsset(
 			InSourceAssetIdx,
 			bInMirrored,
@@ -172,13 +172,13 @@ struct FSearchIndexAsset
 	float GetFirstSampleTime(int32 SchemaSampleRate) const { check(SchemaSampleRate > 0); return FirstSampleIdx / float(SchemaSampleRate); }
 	float GetLastSampleTime(int32 SchemaSampleRate) const { check(SchemaSampleRate > 0); return LastSampleIdx / float(SchemaSampleRate); }
 
-	int32 GetPoseIndexFromTime(float Time, int32 SchemaSampleRate) const
+	int32 GetPoseIndexFromTime(float RealTimeInSeconds, int32 SchemaSampleRate) const
 	{
 #if DO_CHECK
 		check(IsInitialized());
 #endif
 		const int32 NumPoses = GetNumPoses();
-		int32 PoseOffset = FMath::RoundToInt(SchemaSampleRate * Time) - FirstSampleIdx;
+		int32 PoseOffset = FMath::RoundToInt(SchemaSampleRate * RealTimeInSeconds) - FirstSampleIdx;
 		if (bLooping)
 		{
 			PoseOffset = FMath::WrapExclusive(PoseOffset, 0, NumPoses);
diff --git a/Engine/Plugins/Animation/PoseSearch/Source/Runtime/Public/PoseSearch/PoseSearchInteractionIsland.h b/Engine/Plugins/Animation/PoseSearch/Source/Runtime/Public/PoseSearch/PoseSearchInteractionIsland.h
--- a/Engine/Plugins/Animation/PoseSearch/Source/Runtime/Public/PoseSearch/PoseSearchInteractionIsland.h
+++ b/Engine/Plugins/Animation/PoseSearch/Source/Runtime/Public/PoseSearch/PoseSearchInteractionIsland.h
@@ -40,12 +40,15 @@ struct FInteractionSearchContextBase
 #endif // DO_CHECK
 };
 
+// @todo: make FPoseSearchContinuingProperties.PlayingAsset a TWeakObjectPtr and remove this struct!
 // Experimental, this feature might be removed without warning, not for production use
 struct FInteractionSearchContext : FInteractionSearchContextBase
 {
 	// TWeakObjectPtr version of FPoseSearchContinuingProperties to be GC friendly
 	TWeakObjectPtr<const UObject> PlayingAsset = nullptr;
 	float PlayingAssetAccumulatedTime = 0.f;
+	bool bIsPlayingAssetMirrored = false;
+	FVector PlayingAssetBlendParameters = FVector::ZeroVector;
 	EPoseSearchInterruptMode InterruptMode = EPoseSearchInterruptMode::DoNotInterrupt;
 	TArray<int32, TInlineAllocator<PreallocatedRolesNum>> TickPriorities;
 
@@ -54,6 +57,8 @@ struct FInteractionSearchContext : FInteractionSearchContextBase
 		FPoseSearchContinuingProperties ContinuingProperties;
 		ContinuingProperties.PlayingAsset = PlayingAsset.Get();
 		ContinuingProperties.PlayingAssetAccumulatedTime = PlayingAssetAccumulatedTime;
+		ContinuingProperties.bIsPlayingAssetMirrored = bIsPlayingAssetMirrored;
+		ContinuingProperties.PlayingAssetBlendParameters = PlayingAssetBlendParameters;
 		ContinuingProperties.InterruptMode = InterruptMode;
 		return ContinuingProperties;
 	}
diff --git a/Engine/Plugins/Animation/PoseSearch/Source/Runtime/Public/PoseSearch/PoseSearchLibrary.h b/Engine/Plugins/Animation/PoseSearch/Source/Runtime/Public/PoseSearch/PoseSearchLibrary.h
--- a/Engine/Plugins/Animation/PoseSearch/Source/Runtime/Public/PoseSearch/PoseSearchLibrary.h
+++ b/Engine/Plugins/Animation/PoseSearch/Source/Runtime/Public/PoseSearch/PoseSearchLibrary.h
@@ -57,15 +57,21 @@ PRAGMA_ENABLE_DEPRECATION_WARNINGS
 	UE_DEPRECATED(5.6, "Use CurrentSearchResult.UpdateWithNormalizedTime or CurrentSearchResult.UpdateWithRealTime instead")
 	POSESEARCH_API void AdjustAssetTime(float AssetTime);
 
+	UE_DEPRECATED(5.6, "This API is deprecated and it'll be made private")
 	POSESEARCH_API void UpdateWantedPlayRate(const UE::PoseSearch::FSearchContext& SearchContext, const FFloatInterval& PlayRate, float TrajectorySpeedMultiplier, const FPoseSearchEvent& EventToSearch /*= FPoseSearchEvent()*/);
 
 	FVector GetEstimatedFutureRootMotionVelocity() const;
 
+	//@todo: fix all references to CurrentSearchResult before deprecating it
+	//UE_DEPRECATED(5.6, "Use SearchResult instead")
 	UE::PoseSearch::FSearchResult CurrentSearchResult;
 
+	FPoseSearchBlueprintResult SearchResult;
+
 	// Time since the last pose jump
 	float ElapsedPoseSearchTime = 0.f;
 
+	UE_DEPRECATED(5.6, "Use SearchResult.WantedPlayRate instead")
 	// wanted PlayRate to have the selected animation playing at the estimated requested speed from the query.
 	float WantedPlayRate = 1.f;
 
@@ -107,6 +113,13 @@ struct POSESEARCH_API FPoseSearchContinuingProperties
 	UPROPERTY(EditAnywhere, BlueprintReadWrite, Category=State)
 	float PlayingAssetAccumulatedTime = 0.f;
 
+	UPROPERTY(EditAnywhere, BlueprintReadWrite, Category=State)
+	bool bIsPlayingAssetMirrored = false;
+
+	// PlayingAsset associated BlendParameters (if PlayingAsset is a UBlendSpace)
+	UPROPERTY(Transient, VisibleAnywhere, BlueprintReadOnly, Category=State)
+	FVector PlayingAssetBlendParameters = FVector::ZeroVector;
+
 	UPROPERTY(EditAnywhere, BlueprintReadWrite, Category=State)
 	EPoseSearchInterruptMode InterruptMode = EPoseSearchInterruptMode::DoNotInterrupt;
 };
@@ -147,7 +160,9 @@ class POSESEARCH_API UPoseSearchLibrary : public UBlueprintFunctionLibrary
 		bool bSearch,
 		float RecordingTime)
 	{
+		PRAGMA_DISABLE_DEPRECATION_WARNINGS
 		TraceMotionMatching(SearchContext, CurrentState.CurrentSearchResult, CurrentState.ElapsedPoseSearchTime, DeltaTime, bSearch, CurrentState.WantedPlayRate, EPoseSearchInterruptMode::DoNotInterrupt);
+		PRAGMA_ENABLE_DEPRECATION_WARNINGS
 	}
 	
 	UE_DEPRECATED(5.5, "Use TraceMotionMatching with different signature instead.")
diff --git a/Engine/Plugins/Animation/PoseSearch/Source/Runtime/Public/PoseSearch/PoseSearchResult.h b/Engine/Plugins/Animation/PoseSearch/Source/Runtime/Public/PoseSearch/PoseSearchResult.h
--- a/Engine/Plugins/Animation/PoseSearch/Source/Runtime/Public/PoseSearch/PoseSearchResult.h
+++ b/Engine/Plugins/Animation/PoseSearch/Source/Runtime/Public/PoseSearch/PoseSearchResult.h
@@ -9,6 +9,7 @@
 class UPoseSearchDatabase;
 class UPoseSearchSchema;
 struct FGameplayTag;
+struct FPoseSearchBlueprintResult;
 
 namespace UE::PoseSearch
 {
@@ -43,11 +44,17 @@ struct FSearchResult
 	int32 BestPosePos = 0;
 #endif // WITH_EDITOR && ENABLE_ANIM_DEBUG && UE_POSE_SEARCH_TRACE_ENABLED
 
+	// Experimental, this feature might be removed without warning, not for production use
 	// Attempts to set the internal state to match the provided asset time including updating the internal DbPoseIdx. 
 	// If the provided asset time is out of bounds for the currently playing asset then this function will reset the 
 	// state back to the default state.
-	POSESEARCH_API void UpdateWithNormalizedTime(float NormalizedTime);
-	POSESEARCH_API void UpdateWithRealTime(float RealTime);
+	void UpdateWithNormalizedTime(float NormalizedTime);
+	
+	// Experimental, this feature might be removed without warning, not for production use
+	void UpdateWithRealTime(float RealTime);
+
+	// Experimental, this feature might be removed without warning, not for production use
+	bool DebugValidate() const;
 
 	bool IsValid() const
 	{
@@ -90,8 +97,6 @@ struct FSearchResult
 	// Experimental, this feature might be removed without warning, not for production use
 	POSESEARCH_API const UAnimationAsset* GetCurrentResultAnimationAsset(const FRole& Role) const;
 	
-	bool CanAdvance(float DeltaTime) const;
-
 	bool operator==(const FSearchResult& Other) const
 	{
 		// best cost of the currently selected PoseIdx (it could be equal to ContinuingPoseCost)
@@ -103,6 +108,9 @@ struct FSearchResult
 			AssetTime == Other.AssetTime &&
 			bIsContinuingPoseSearch == Other.bIsContinuingPoseSearch;
 	}
+
+	// Experimental, this feature might be removed without warning, not for production use
+	void InitFrom(const FPoseSearchBlueprintResult& BlueprintResult);
 };
 
 } // namespace UE::PoseSearch
@@ -175,5 +183,11 @@ public:
 	// AnimContexts involved in this interaction for all the roles in SelectedAnim (as UMultiAnimAsset). If the role is missing because optional the AnimContext will be null
 	// BE CAREFUL HOW YOU USE THEM! since accessing data from other actors without proper tick dependencies leads to threading problems!
 	TArray<TWeakObjectPtr<const UObject>> AnimContexts;
+
+	// Experimental, this feature might be removed without warning, not for production use
+	bool InitFrom(const UE::PoseSearch::FSearchResult& SearchResult, float InWantedPlayRate);
+
+	// Experimental, this feature might be removed without warning, not for production use
+	UAnimationAsset* GetAnimationAssetForRole();
 };
 
diff --git a/Engine/Plugins/Chooser/Source/Chooser/Private/AnimNode_ChooserPlayer.cpp b/Engine/Plugins/Chooser/Source/Chooser/Private/AnimNode_ChooserPlayer.cpp
--- a/Engine/Plugins/Chooser/Source/Chooser/Private/AnimNode_ChooserPlayer.cpp
+++ b/Engine/Plugins/Chooser/Source/Chooser/Private/AnimNode_ChooserPlayer.cpp
@@ -47,6 +47,8 @@ UAnimationAsset* FAnimNode_ChooserPlayer::ChooseAsset(const FAnimationUpdateCont
 				IPoseSearchProvider::FSearchPlayingAsset PlayingAsset;
 				PlayingAsset.Asset = GetAnimAsset();
 				PlayingAsset.AccumulatedTime = GetAccumulatedTime();
+				PlayingAsset.bMirrored = GetMirror();
+				PlayingAsset.BlendParameters = GetBlendParameters();
 
 				const IPoseSearchProvider::FSearchResult SearchResult = PoseSearchProvider->Search(Context, ChosenAssets, PlayingAsset, IPoseSearchProvider::FSearchFutureAsset() );
 				if (UAnimationAsset* SelectedAnimationAsset = Cast<UAnimationAsset>(SearchResult.SelectedAsset))
diff --git a/Engine/Plugins/Experimental/AnimNextPoseSearch/Source/AnimNextPoseSearch/Private/MotionMatchingTrait.cpp b/Engine/Plugins/Experimental/AnimNextPoseSearch/Source/AnimNextPoseSearch/Private/MotionMatchingTrait.cpp
--- a/Engine/Plugins/Experimental/AnimNextPoseSearch/Source/AnimNextPoseSearch/Private/MotionMatchingTrait.cpp
+++ b/Engine/Plugins/Experimental/AnimNextPoseSearch/Source/AnimNextPoseSearch/Private/MotionMatchingTrait.cpp
@@ -46,13 +46,6 @@ static FAutoConsoleVariableRef CVarAnimNextMotionMatchingPlayRateEnabled(
 	TEXT("1: Enable all usages of PlayRate.\n")
 	TEXT("2: Enable PlayRate in PoseSeach only (Not used in actual playback).\n")
 );
-
-static bool GVarAnimNextMotionMatchingUseTimelineRealTime = false;
-static FAutoConsoleVariableRef CVarAnimNextMotionMatchingUseTimelineRealTime(
-	TEXT("a.AnimNext.MotionMatchingTrait.UseTimelineRealTime"),
-	GVarAnimNextMotionMatchingUseTimelineRealTime,
-	TEXT("Debug code to use the real time from the timeline rather the reconstructed normalized time for blend spaces to update the continuing pose")
-);
 #endif // ENABLE_ANIM_DEBUG
 
 namespace UE::AnimNext
@@ -121,12 +114,6 @@ namespace UE::AnimNext
 		}
 
 		const FSharedData* SharedData = Binding.GetSharedData<FSharedData>();
-		const UObject* AnimationAsset = nullptr;
-		float AssetTime = 0.f;
-		float WantedPlayRate = 1.f;
-		bool bLoop = false;
-		FVector BlendSpaceParameters = FVector::ZeroVector;
-
 		const UObject* AnimContext = nullptr;
 		if (const FAnimNextModuleInstance* ModuleInstance = Context.GetRootGraphInstance().GetModuleInstance())
 		{
@@ -134,161 +121,101 @@ namespace UE::AnimNext
 		}
 
 		FMotionMatchingState& MotionMatchingState = InstanceData->MotionMatchingState;
-		FPoseSearchBlueprintResult& MotionMatchingInteractionState = InstanceData->MotionMatchingInteractionState;
 
-		const TArray<FPoseSearchInteractionAvailability>& Availabilities = SharedData->GetAvailabilities(Binding);
-		if (Availabilities.IsEmpty())
+		// synchronizing with GetAccumulatedTime or resetting MotionMatchingState, and conditionally resetting FAnimNode_BlendStack_Standalone
+		// @todo: implement this MotionMatchingState.Reset() condition for parity to FAnimNode_MotionMatching::UpdateAssetPlayer
+		//if (bResetOnBecomingRelevant && UpdateCounter.HasEverBeenUpdated() && !UpdateCounter.WasSynchronizedCounter(Context.AnimInstanceProxy->GetUpdateCounter()))
+		//{
+		//	// If we just became relevant and haven't been initialized yet, then reset motion matching state, otherwise update the asset time using the player node.
+		//	MotionMatchingState.Reset();
+		//	FAnimNode_BlendStack_Standalone::Reset();
+		//}
+		//else 
+		if (MotionMatchingState.SearchResult.SelectedDatabase == nullptr || MotionMatchingState.SearchResult.SelectedDatabase->Schema == nullptr)
 		{
-			MotionMatchingInteractionState = FPoseSearchBlueprintResult();
 		}
-		else
+#if WITH_EDITOR
+		else if (EAsyncBuildIndexResult::Success != FAsyncPoseSearchDatabasesManagement::RequestAsyncBuildIndex(MotionMatchingState.SearchResult.SelectedDatabase.Get(), ERequestAsyncBuildFlag::ContinueRequest))
 		{
-			// WIP motion matched interactions
-			CheckInteractionThreadSafety(AnimContext);
+			// MotionMatchingState.SearchResult.Database is indexing, and it's not safe to use its previous index properties cached in MotionMatchingState
+			MotionMatchingState.Reset();
+		}
+#endif // WITH_EDITOR
 
-			// We adjust the MotionMatchingInteractionState time to the current player node's asset time. This is done 
-			// because the player node may have ticked more or less time than we expected due to variable dt or the 
-			// dynamic playback rate adjustment and as such the MotionMatchingInteractionState does not update by itself
-			const FTimelineState TimelineState = TimelineTrait.GetState(Context);
-			float TimelinePosition = TimelineState.GetPosition();
+		// We adjust the MotionMatchingInteractionState time to the current player node's asset time. This is done 
+		// because the player node may have ticked more or less time than we expected due to variable dt or the 
+		// dynamic playback rate adjustment and as such the MotionMatchingInteractionState does not update by itself
+		const FTimelineState TimelineState = TimelineTrait.GetState(Context);
+		const float TimelineRealTime = TimelineState.GetPosition();
+		float TimelineNormalizedTime = TimelineRealTime;
 
-			// @TODO: This is a hack, since we don't have a way to get the normalized timeline of a blendspace yet (pending syncing...?)
-			if (const UE::PoseSearch::FSearchIndexAsset* SearchIndexAsset = MotionMatchingState.CurrentSearchResult.GetSearchIndexAsset())
+		// @todo: This is a hack, since we don't have a way to get the normalized timeline of a blendspace yet (pending syncing...?)
+		if (Cast<UBlendSpace>(MotionMatchingState.SearchResult.GetAnimationAssetForRole()))
+		{
+			// Convert to normalized time.
+			const float TimelineDuration = TimelineState.GetDuration();
+			if (TimelineDuration > 0.0f)
 			{
-				if (const UMultiAnimAsset* MultiAnimAsset = Cast<UMultiAnimAsset>(MotionMatchingInteractionState.SelectedAnim))
-				{
-					if (Cast<UBlendSpace>(MultiAnimAsset->GetAnimationAsset(MotionMatchingState.CurrentSearchResult.Database->Schema->GetDefaultRole())))
-					{
-						// Convert to normalized time.
-						const float TimelineDuration = TimelineState.GetDuration();
-						if (TimelineDuration > 0.0f)
-						{
-							// NOTE: This doesn't work because AnimNext normalized time does not much database indexing normalized time.
-							TimelinePosition /= TimelineState.GetDuration();
-						}
-					}
-				}
+				// NOTE: This doesn't work because AnimNext normalized time does not match database indexing normalized time.
+				TimelineNormalizedTime = TimelineRealTime / TimelineDuration;
 			}
-
-			// we're NOT providing the UMultiAnimAsset here, relying on the interaction subsystem to resolve the ContinuingProperties properly
-			FPoseSearchContinuingProperties ContinuingProperties;
-			ContinuingProperties.PlayingAsset = MotionMatchingInteractionState.SelectedAnim.Get(); // GetAnimAsset();
-			ContinuingProperties.PlayingAssetAccumulatedTime = TimelinePosition;
-
-			UPoseSearchInteractionLibrary::MotionMatchInteraction(MotionMatchingInteractionState, Availabilities, AnimContext, ContinuingProperties, FName(), PoseHistory, SharedData->GetbValidateResultAgainstAvailabilities(Binding));
-			check(MotionMatchingInteractionState.ActorRootTransforms.Num() == MotionMatchingInteractionState.ActorRootBoneTransforms.Num());
 		}
 
-		if (MotionMatchingInteractionState.SelectedAnim != nullptr)
+#if ENABLE_VISUAL_LOG
+		if (FVisualLogger::IsRecording())
 		{
-			if (!MotionMatchingInteractionState.bIsContinuingPoseSearch)
-			{
-				if (const UMultiAnimAsset* MultiAnimAsset = Cast<UMultiAnimAsset>(MotionMatchingInteractionState.SelectedAnim))
-				{
-					AnimationAsset = MultiAnimAsset->GetAnimationAsset(MotionMatchingInteractionState.Role);
-				}
-			}
-			AssetTime = MotionMatchingInteractionState.SelectedTime;
-			WantedPlayRate = MotionMatchingInteractionState.WantedPlayRate;
-			bLoop = MotionMatchingInteractionState.bLoop;
-			BlendSpaceParameters = MotionMatchingInteractionState.BlendParameters;
-
-			MotionMatchingState.Reset();
+			const FName AnimName = TimelineState.GetDebugName();
+			const float Duration = TimelineState.GetDuration();
+			static const TCHAR* LogName = TEXT("FMotionMatchingTrait");
+			UE_VLOG(Context.GetHostObject(), LogName, Verbose, TEXT("TimelineRealTime: %f"), TimelineRealTime);
+			UE_VLOG(Context.GetHostObject(), LogName, Verbose, TEXT("TimelineNormalizedTime: %f"), TimelineNormalizedTime);
+			UE_VLOG(Context.GetHostObject(), LogName, Verbose, TEXT("AnimName: %s"), *AnimName.ToString());
+			UE_VLOG(Context.GetHostObject(), LogName, Verbose, TEXT("Duration: %f"), Duration);
 		}
-		else
+#endif
+
+		bool bIsInteraction = false;
+		const TArray<FPoseSearchInteractionAvailability>& Availabilities = SharedData->GetAvailabilities(Binding);
+		if (!Availabilities.IsEmpty())
 		{
-			// synchronizing with GetAccumulatedTime or resetting MotionMatchingState, and conditionally resetting FAnimNode_BlendStack_Standalone
-			// @todo: implement this MotionMatchingState.Reset() condition for parity to FAnimNode_MotionMatching::UpdateAssetPlayer
-			//if (bResetOnBecomingRelevant && UpdateCounter.HasEverBeenUpdated() && !UpdateCounter.WasSynchronizedCounter(Context.AnimInstanceProxy->GetUpdateCounter()))
-			//{
-			//	// If we just became relevant and haven't been initialized yet, then reset motion matching state, otherwise update the asset time using the player node.
-			//	MotionMatchingState.Reset();
-			//	FAnimNode_BlendStack_Standalone::Reset();
-			//}
-			//else 
-			if (MotionMatchingState.CurrentSearchResult.Database == nullptr || MotionMatchingState.CurrentSearchResult.Database->Schema == nullptr)
+			// WIP motion matched interactions
+			CheckInteractionThreadSafety(AnimContext);
+
+			FPoseSearchContinuingProperties ContinuingProperties;
+			if (MotionMatchingState.SearchResult.bIsInteraction)
 			{
+				// @todo: retrieve the playing UMultiAnimAsset from the UAnimAsset blendspace is playing via GetAnimAsset() and blend parameters via GetBlendParameters() and mirror state
+				ContinuingProperties.PlayingAsset = MotionMatchingState.SearchResult.SelectedAnim.Get();
+				ContinuingProperties.PlayingAssetAccumulatedTime = TimelineNormalizedTime;
+				ContinuingProperties.bIsPlayingAssetMirrored = MotionMatchingState.SearchResult.bIsMirrored;
+				ContinuingProperties.PlayingAssetBlendParameters = MotionMatchingState.SearchResult.BlendParameters;
 			}
 
-#if WITH_EDITOR
-			else if (EAsyncBuildIndexResult::Success != FAsyncPoseSearchDatabasesManagement::RequestAsyncBuildIndex(MotionMatchingState.CurrentSearchResult.Database.Get(), ERequestAsyncBuildFlag::ContinueRequest))
+			FPoseSearchBlueprintResult SearchResult;
+			UPoseSearchInteractionLibrary::MotionMatchInteraction(SearchResult, Availabilities, AnimContext, ContinuingProperties, FName(), PoseHistory, SharedData->GetbValidateResultAgainstAvailabilities(Binding));
+			check(SearchResult.ActorRootTransforms.Num() == SearchResult.ActorRootBoneTransforms.Num());
+
+			if (SearchResult.SelectedAnim.IsValid())
 			{
-				// MotionMatchingState.CurrentSearchResult.Database is indexing, and it's not safe to use its previous index properties cached in MotionMatchingState
-				MotionMatchingState.Reset();
+				bIsInteraction = true;
+				MotionMatchingState.SearchResult = SearchResult;
 			}
-			// in case this node is not updated, and MotionMatchingState.CurrentSearchResult.Database gets modified, we could end up with CurrentSearchResult being out of synch with the updated database, so we need to reset the state
-			else if (MotionMatchingState.CurrentSearchResult.IsValid() && MotionMatchingState.CurrentSearchResult.PoseIdx >= MotionMatchingState.CurrentSearchResult.Database->GetSearchIndex().GetNumPoses())
+		}
+
+		// performing the regular single character motion matching search in case there's no MM interaction
+		if (!bIsInteraction)
+		{
+			if (MotionMatchingState.SearchResult.bIsInteraction)
 			{
-				// in case this node is not updated, and MotionMatchingState.CurrentSearchResult.Database gets modified, we could end up with
-				// CurrentSearchResult being out of synch with the updated database, so we need to reset the state
-				MotionMatchingState.Reset();
+				MotionMatchingState.SearchResult = FPoseSearchBlueprintResult();
 			}
-#endif // WITH_EDITOR
-			// @todo: implement this MotionMatchingState.Reset() condition for parity to FAnimNode_MotionMatching::UpdateAssetPlayer
-			//else if (GetAnimAsset() != MotionMatchingState.CurrentSearchResult.GetCurrentResultAnimationAsset())
-			//{
-			//	// MotionMatchingState and blendstack are not in synch! MotionMatchingState thinks it's playing MotionMatchingState.CurrentSearchResult.GetCurrentResultAnimationAsset()
-			//	// and blendstack thinks it's playing GetAnimAsset(). The two could have gone out of synch becasue or external blendstack requests, and so we reset the MotionMatchingState
-			//	MotionMatchingState.Reset();
-			//}
 			else
 			{
-				// We adjust the MotionMatchingState asset time to the current player node's asset time. This is done 
-				// because the player node may have ticked more or less time than we expected due to variable dt or the 
-				// dynamic playback rate adjustment and as such the MotionMatchingState does not update by itself
-				const FTimelineState TimelineState = TimelineTrait.GetState(Context);
-				const float TimelineRealTime = TimelineState.GetPosition();
-				float TimelineNormalizedTime = -1.f;
-
-				bool bUseTimelineRealTime = false;
-#if ENABLE_ANIM_DEBUG
-				bUseTimelineRealTime = GVarAnimNextMotionMatchingUseTimelineRealTime;
-#endif // ENABLE_ANIM_DEBUG
-
-				if (bUseTimelineRealTime)
-				{
-					MotionMatchingState.CurrentSearchResult.UpdateWithRealTime(TimelineRealTime);
-				}
-				else
-				{
-					TimelineNormalizedTime = TimelineRealTime;
-
-					// @todo: This is a hack, since we don't have a way to get the normalized timeline of a blendspace yet (pending syncing...?)
-					if (const UE::PoseSearch::FSearchIndexAsset* SearchIndexAsset = MotionMatchingState.CurrentSearchResult.GetSearchIndexAsset())
-					{
-						if (const FPoseSearchDatabaseAnimationAssetBase* DatabaseAnimationAssetBase = MotionMatchingState.CurrentSearchResult.Database->GetDatabaseAnimationAsset<FPoseSearchDatabaseAnimationAssetBase>(*SearchIndexAsset))
-						{
-							if (Cast<UBlendSpace>(DatabaseAnimationAssetBase->GetAnimationAssetForRole(MotionMatchingState.CurrentSearchResult.Database->Schema->GetDefaultRole())))
-							{
-								// Convert to normalized time.
-								const float TimelineDuration = TimelineState.GetDuration();
-								if (TimelineDuration > 0.0f)
-								{
-									// NOTE: This doesn't work because AnimNext normalized time does not much database indexing normalized time.
-									TimelineNormalizedTime /= TimelineDuration;
-								}
-							}
-						}
-					}
-
-					MotionMatchingState.CurrentSearchResult.UpdateWithNormalizedTime(TimelineNormalizedTime);
-				}
-
-#if ENABLE_VISUAL_LOG
-				if (FVisualLogger::IsRecording())
-				{
-					//const float AfterNormalize = TimelinePosition;
-					const float AfterMotionMatchingStateStuff = MotionMatchingState.CurrentSearchResult.AssetTime;
-					const FName AnimName = TimelineState.GetDebugName();
-					const float Duration = TimelineState.GetDuration();
-					static const TCHAR* LogName = TEXT("FMotionMatchingTrait");
-					UE_VLOG(Context.GetHostObject(), LogName, Verbose, TEXT("TimelineRealTime: %f"), TimelineRealTime);
-					UE_VLOG(Context.GetHostObject(), LogName, Verbose, TEXT("TimelineNormalizedTime: %f"), TimelineNormalizedTime);
-					UE_VLOG(Context.GetHostObject(), LogName, Verbose, TEXT("AfterMotionMatchingStateStuff: %f"), AfterMotionMatchingStateStuff);
-					UE_VLOG(Context.GetHostObject(), LogName, Verbose, TEXT("AnimName: %s"), *AnimName.ToString());
-					UE_VLOG(Context.GetHostObject(), LogName, Verbose, TEXT("Duration: %f"), Duration);
-				}
-#endif
+				// @todo: ask blendstack what it's playing instead of relying on MotionMatchingInteractionState.SelectedAnim
+				//MotionMatchingState.SearchResult.SelectedAnim = GetAnimAsset();
+				MotionMatchingState.SearchResult.SelectedTime = TimelineNormalizedTime;
+				// @todo: ask blendstack what the blend parameters are
+				//MotionMatchingState.SearchResult.BlendParameters = GetBlendParameters();
 			}
 
 			const TConstArrayView<TObjectPtr<const UPoseSearchDatabase>> DatabasesToSearch = SharedData->GetDatabases(Binding);
@@ -324,121 +251,102 @@ namespace UE::AnimNext
 					// @todo: add debug drawing suport
 	//#if ENABLE_ANIM_DEBUG
 	//				, CVarAnimNodeMotionMatchingDrawQuery.GetValueOnAnyThread()
-	//				, CVarAnimNodeMotionMatchingDrawCurResult.GetValueOnAnyThread()
+	//				, CVarAnimNodeMotionMatchingDrawCurResult.GetValueOnAnyThread() 
 	//#endif // ENABLE_ANIM_DEBUG
 				);
-				WantedPlayRate = MotionMatchingState.WantedPlayRate;
-			}
-
-			const UPoseSearchDatabase* Database = MotionMatchingState.CurrentSearchResult.Database.Get();
-#if WITH_EDITOR
-			// we need to validate if Database is not reindexing because of any possible throttling of the search!
-			if (EAsyncBuildIndexResult::Success == FAsyncPoseSearchDatabasesManagement::RequestAsyncBuildIndex(Database, ERequestAsyncBuildFlag::ContinueRequest))
-#endif // WITH_EDITOR
-			{
-				const bool bJumpToPose = MotionMatchingState.CurrentSearchResult.IsValid() && !MotionMatchingState.CurrentSearchResult.bIsContinuingPoseSearch;
-				if (bJumpToPose)
-				{
-					if (const UE::PoseSearch::FSearchIndexAsset* SearchIndexAsset = MotionMatchingState.CurrentSearchResult.GetSearchIndexAsset())
-					{
-						if (const FPoseSearchDatabaseAnimationAssetBase* DatabaseAnimationAssetBase = Database->GetDatabaseAnimationAsset<FPoseSearchDatabaseAnimationAssetBase>(*SearchIndexAsset))
-						{
-							AnimationAsset = DatabaseAnimationAssetBase->GetAnimationAssetForRole(Database->Schema->GetDefaultRole());
-							AssetTime = MotionMatchingState.CurrentSearchResult.AssetTime;
-							bLoop = SearchIndexAsset->IsLooping();
-							BlendSpaceParameters = SearchIndexAsset->GetBlendParameters();
-						}
-					}
-				}
 			}
 		}
 
-		float DesiredPlayRate = WantedPlayRate;
+		float DesiredPlayRate = MotionMatchingState.SearchResult.WantedPlayRate;
 #if ENABLE_ANIM_DEBUG
 		DesiredPlayRate = GVarAnimNextMotionMatchingPlayRateEnabled == UE::Private::EPlayRateState::Enabled ? DesiredPlayRate : 1.0f;
 #endif // ENABLE_ANIM_DEBUG
 
-		if (AnimationAsset)
+		const bool bJumpToPose = MotionMatchingState.SearchResult.SelectedAnim.IsValid() && !MotionMatchingState.SearchResult.bIsContinuingPoseSearch;
+		if (bJumpToPose)
 		{
-			const bool bTrySkipBlendsForBlendSpaces = SharedData->GetbTrySkipBlendsForBlendSpaces(Binding);
-			const bool bIsBlendSpace = AnimationAsset->IsA<UBlendSpace>();
-			bool bSkipBlend = false;
-			if (bIsBlendSpace && bTrySkipBlendsForBlendSpaces)
+			if (UAnimationAsset* AnimationAsset = MotionMatchingState.SearchResult.GetAnimationAssetForRole())
 			{
-				IBlendStack::FGraphRequestPtr ActiveGraphRequest;
-				BlendStackBinding.GetActiveGraph(Context, ActiveGraphRequest);
-				if (ActiveGraphRequest && (ActiveGraphRequest->FactoryObject == AnimationAsset))
+				const bool bTrySkipBlendsForBlendSpaces = SharedData->GetbTrySkipBlendsForBlendSpaces(Binding);
+				const bool bIsBlendSpace = AnimationAsset->IsA<UBlendSpace>();
+				bool bSkipBlend = false;
+				if (bIsBlendSpace && bTrySkipBlendsForBlendSpaces)
 				{
-					check(ActiveGraphRequest->AnimationGraph);
-					for (FStructView& Payload : ActiveGraphRequest->GraphPayload.Get())
+					IBlendStack::FGraphRequestPtr ActiveGraphRequest;
+					BlendStackBinding.GetActiveGraph(Context, ActiveGraphRequest);
+					if (ActiveGraphRequest && (ActiveGraphRequest->FactoryObject == AnimationAsset))
 					{
-						if (FAnimNextNativeDataInterface_BlendSpacePlayer* BlendSpacePayload = Payload.GetPtr<FAnimNextNativeDataInterface_BlendSpacePlayer>())
+						check(ActiveGraphRequest->AnimationGraph);
+						for (FStructView& Payload : ActiveGraphRequest->GraphPayload.Get())
 						{
-							const float MaxDeltaAssetTimeToTrySkipBlendsForBlendSpaces = SharedData->GetMaxDeltaAssetTimeToTrySkipBlendsForBlendSpaces(Binding);
-							if (FMath::Abs(BlendSpacePayload->StartPosition - AssetTime) < MaxDeltaAssetTimeToTrySkipBlendsForBlendSpaces)
+							if (FAnimNextNativeDataInterface_BlendSpacePlayer* BlendSpacePayload = Payload.GetPtr<FAnimNextNativeDataInterface_BlendSpacePlayer>())
 							{
-								// If we're still on the same blend space, and we're updating blend space inputs dynamically, then DON'T BLEND.
-								bSkipBlend = true;
+								const float MaxDeltaAssetTimeToTrySkipBlendsForBlendSpaces = SharedData->GetMaxDeltaAssetTimeToTrySkipBlendsForBlendSpaces(Binding);
+								if (FMath::Abs(BlendSpacePayload->StartPosition - MotionMatchingState.SearchResult.SelectedTime) < MaxDeltaAssetTimeToTrySkipBlendsForBlendSpaces)
+								{
+									// If we're still on the same blend space, and we're updating blend space inputs dynamically, then DON'T BLEND.
+									bSkipBlend = true;
 
-								BlendSpacePayload->PlayRate = DesiredPlayRate;
+									BlendSpacePayload->PlayRate = DesiredPlayRate;
 
-								BlendSpacePayload->XAxisSamplePoint = BlendSpaceParameters.X;
-								BlendSpacePayload->YAxisSamplePoint = BlendSpaceParameters.Y;
+									BlendSpacePayload->XAxisSamplePoint = MotionMatchingState.SearchResult.BlendParameters.X;
+									BlendSpacePayload->YAxisSamplePoint = MotionMatchingState.SearchResult.BlendParameters.Y;
+								}
+								break;
 							}
-							break;
 						}
 					}
 				}
-			}
 
-			if (!bSkipBlend)
-			{
-				UE::AnimNext::IBlendStack::FGraphRequest NewGraphRequest;
-				NewGraphRequest.BlendArgs = SharedData->BlendArguments;
-
-				// Create a default 'native data interface' payload via which we expect to communicate with the graph 
-				FAnimNextDataInterfacePayload GraphPayload;
-				FInstancedStruct AssetPlayerPayload;
-				// @TODO: Support asset->playload mapping in a more data-driven way
-				if (bIsBlendSpace)
+				if (!bSkipBlend)
 				{
-					// Tell blendspace samples to synchronize among themselves, but not with anything else.
-					FInstancedStruct SyncGroupPayload;
-					auto& NativeInterface = SyncGroupPayload.InitializeAs<FAnimNextNativeDataInterface_SynchronizeUsingGroups>();
-					NativeInterface.GroupName = NAME_None;
-					NativeInterface.GroupRole = EAnimGroupSynchronizationRole::AlwaysFollower;
-					NativeInterface.SyncMode = SharedData->GetSyncMode(Binding);
-					NativeInterface.MatchSyncPoint = true;
-					GraphPayload.AddNative(MoveTemp(SyncGroupPayload));
-
-					FAnimNextNativeDataInterface_BlendSpacePlayer& NewNativeInterface = AssetPlayerPayload.InitializeAs<FAnimNextNativeDataInterface_BlendSpacePlayer>();
-					NewNativeInterface.BlendSpace = Cast<UBlendSpace>(AnimationAsset);
-
-					NewNativeInterface.XAxisSamplePoint = BlendSpaceParameters.X;
-					NewNativeInterface.YAxisSamplePoint = BlendSpaceParameters.Y;
-
-					NewNativeInterface.PlayRate = DesiredPlayRate;
-					NewNativeInterface.StartPosition = AssetTime;
-					NewNativeInterface.Loop = bLoop;
-				}
-				else
-				{
-					FAnimNextNativeDataInterface_AnimSequencePlayer& NewNativeInterface = AssetPlayerPayload.InitializeAs<FAnimNextNativeDataInterface_AnimSequencePlayer>();
-					NewNativeInterface.AnimSequence = Cast<UAnimSequence>(AnimationAsset);
-					NewNativeInterface.PlayRate = DesiredPlayRate;
-					NewNativeInterface.StartPosition = AssetTime;
-					NewNativeInterface.Loop = bLoop;
-				}
+					UE::AnimNext::IBlendStack::FGraphRequest NewGraphRequest;
+					NewGraphRequest.BlendArgs = SharedData->BlendArguments;
+
+					// Create a default 'native data interface' payload via which we expect to communicate with the graph 
+					FAnimNextDataInterfacePayload GraphPayload;
+					FInstancedStruct AssetPlayerPayload;
+					// @TODO: Support asset->playload mapping in a more data-driven way
+					if (bIsBlendSpace)
+					{
+						// Tell blendspace samples to synchronize among themselves, but not with anything else.
+						FInstancedStruct SyncGroupPayload;
+						auto& NativeInterface = SyncGroupPayload.InitializeAs<FAnimNextNativeDataInterface_SynchronizeUsingGroups>();
+						NativeInterface.GroupName = NAME_None;
+						NativeInterface.GroupRole = EAnimGroupSynchronizationRole::AlwaysFollower;
+						NativeInterface.SyncMode = SharedData->GetSyncMode(Binding);
+						NativeInterface.MatchSyncPoint = true;
+						GraphPayload.AddNative(MoveTemp(SyncGroupPayload));
+
+						FAnimNextNativeDataInterface_BlendSpacePlayer& NewNativeInterface = AssetPlayerPayload.InitializeAs<FAnimNextNativeDataInterface_BlendSpacePlayer>();
+						NewNativeInterface.BlendSpace = Cast<UBlendSpace>(AnimationAsset);
+
+						NewNativeInterface.XAxisSamplePoint = MotionMatchingState.SearchResult.BlendParameters.X;
+						NewNativeInterface.YAxisSamplePoint = MotionMatchingState.SearchResult.BlendParameters.Y;
+
+						NewNativeInterface.PlayRate = DesiredPlayRate;
+						NewNativeInterface.StartPosition = MotionMatchingState.SearchResult.SelectedTime;
+						NewNativeInterface.Loop = MotionMatchingState.SearchResult.bLoop;
+					}
+					else
+					{
+						FAnimNextNativeDataInterface_AnimSequencePlayer& NewNativeInterface = AssetPlayerPayload.InitializeAs<FAnimNextNativeDataInterface_AnimSequencePlayer>();
+						NewNativeInterface.AnimSequence = Cast<UAnimSequence>(AnimationAsset);
+						NewNativeInterface.PlayRate = DesiredPlayRate;
+						NewNativeInterface.StartPosition = MotionMatchingState.SearchResult.SelectedTime;
+						NewNativeInterface.Loop = MotionMatchingState.SearchResult.bLoop;
+					}
 
-				GraphPayload.AddNative(MoveTemp(AssetPlayerPayload));
+					GraphPayload.AddNative(MoveTemp(AssetPlayerPayload));
 
-				const UAnimNextAnimationGraph* AnimationGraph = IGraphFactory::GetGraphFromObjectWithFallback(Context, Binding, AnimationAsset, GraphPayload);
-				if (AnimationGraph != nullptr)
-				{
-					NewGraphRequest.GraphPayload = MoveTemp(GraphPayload);
-					NewGraphRequest.AnimationGraph = AnimationGraph;
-					NewGraphRequest.FactoryObject = AnimationAsset;
-					BlendStackBinding.PushGraph(Context, MoveTemp(NewGraphRequest));
+					const UAnimNextAnimationGraph* AnimationGraph = IGraphFactory::GetGraphFromObjectWithFallback(Context, Binding, AnimationAsset, GraphPayload);
+					if (AnimationGraph != nullptr)
+					{
+						NewGraphRequest.GraphPayload = MoveTemp(GraphPayload);
+						NewGraphRequest.AnimationGraph = AnimationGraph;
+						NewGraphRequest.FactoryObject = AnimationAsset;
+						BlendStackBinding.PushGraph(Context, MoveTemp(NewGraphRequest));
+					}
 				}
 			}
 		}
@@ -486,7 +394,6 @@ namespace UE::AnimNext
 		const FSharedData* SharedData = Binding.GetSharedData<FSharedData>();
 		const FInstanceData* InstanceData = Binding.GetInstanceData<FInstanceData>();
 		const FMotionMatchingState& MotionMatchingState = InstanceData->MotionMatchingState;
-		const FPoseSearchBlueprintResult& MotionMatchingInteractionState = InstanceData->MotionMatchingInteractionState;
 
 		FName ResultOutputVariableName = SharedData->GetMotionMatchingResultVariableName(Binding);
 		if (ResultOutputVariableName == NAME_None)
@@ -494,68 +401,41 @@ namespace UE::AnimNext
 			return;
 		}
 
-		FPoseSearchBlueprintResult Result;
-
-		const bool bIsMotionMatchingInteractionStateValid = MotionMatchingInteractionState.SelectedAnim != nullptr;
-		const UPoseSearchDatabase* CurrentResultDatabase = bIsMotionMatchingInteractionStateValid ? MotionMatchingInteractionState.SelectedDatabase.Get() : MotionMatchingState.CurrentSearchResult.Database.Get();
+		const UPoseSearchDatabase* CurrentResultDatabase = MotionMatchingState.SearchResult.SelectedDatabase.Get();
 		if (CurrentResultDatabase == nullptr || CurrentResultDatabase->Schema == nullptr)
 		{
 			// @todo: should we log a warning?
 		}
-
 #if WITH_EDITOR
 		else if (EAsyncBuildIndexResult::Success != FAsyncPoseSearchDatabasesManagement::RequestAsyncBuildIndex(CurrentResultDatabase, ERequestAsyncBuildFlag::ContinueRequest))
 		{
 			UE_LOG(LogPoseSearch, Warning, TEXT("FMotionMatchingTrait::PreUpdate called while '%s' is indexing. returning an invalid result.."), *CurrentResultDatabase->GetName());
 		}
-		// in case this node is not updated, and MotionMatchingState.CurrentSearchResult.Database gets modified, we could end up with CurrentSearchResult being out of synch with the updated database, so we need to reset the state
-		else if (!bIsMotionMatchingInteractionStateValid && MotionMatchingState.CurrentSearchResult.IsValid() && MotionMatchingState.CurrentSearchResult.PoseIdx >= MotionMatchingState.CurrentSearchResult.Database->GetSearchIndex().GetNumPoses())
-		{
-			UE_LOG(LogPoseSearch, Warning, TEXT("FMotionMatchingTrait::PreUpdate inconsistency between MotionMatchingState and '%s', probably because of stale MotionMatchingState after database got reindexed. returning an invalid result.."), *CurrentResultDatabase->GetName());
-		}
 #endif // WITH_EDITOR
-		else if (bIsMotionMatchingInteractionStateValid)
-		{
-			Result = MotionMatchingInteractionState;
-		}
-		else if (const FSearchIndexAsset* SearchIndexAsset = MotionMatchingState.CurrentSearchResult.GetSearchIndexAsset())
-		{
-			if (const FPoseSearchDatabaseAnimationAssetBase* DatabaseAsset = CurrentResultDatabase->GetDatabaseAnimationAsset<FPoseSearchDatabaseAnimationAssetBase>(*SearchIndexAsset))
-			{
-				Result.SelectedAnim = DatabaseAsset->GetAnimationAsset();
-				Result.SelectedTime = MotionMatchingState.CurrentSearchResult.AssetTime;
-				Result.bIsContinuingPoseSearch = MotionMatchingState.CurrentSearchResult.bIsContinuingPoseSearch;
-				Result.bLoop = SearchIndexAsset->IsLooping();
-				Result.bIsMirrored = SearchIndexAsset->IsMirrored();
-				Result.BlendParameters = SearchIndexAsset->GetBlendParameters();
-				Result.SelectedDatabase = CurrentResultDatabase;
-				Result.WantedPlayRate = MotionMatchingState.WantedPlayRate;
-				Result.SearchCost = MotionMatchingState.CurrentSearchResult.PoseCost;
-			}
-		}
-
+		
 		FAnimNextModuleInstance* ModuleInstance = Binding.GetTraitPtr().GetNodeInstance()->GetOwner().GetModuleInstance();
 		check(ModuleInstance);
 
+		// @todo: merge this duplicated code with the block below
 		const FAnimNextModuleHandle ResultModuleHandle = SharedData->GetMotionMatchingResultModuleHandle(Binding);
 		if (ResultModuleHandle.ModuleHandle.IsValid())
 		{
 			ModuleInstance->QueueTaskOnOtherModule(
 				ResultModuleHandle.ModuleHandle,
 				FRigUnit_AnimNextPrePhysicsEvent::DefaultEventName,
-				[Result, ResultOutputVariableName](const UE::AnimNext::FModuleTaskContext& TaskContext)
+				[&MotionMatchingState, ResultOutputVariableName](const UE::AnimNext::FModuleTaskContext& TaskContext)
 				{
 					FAnimNextModuleInstance* OtherModuleInstance = TaskContext.GetModuleInstance();
 					if (OtherModuleInstance)
 					{
-						OtherModuleInstance->SetVariable(ResultOutputVariableName, Result);
+						OtherModuleInstance->SetVariable(ResultOutputVariableName, MotionMatchingState.SearchResult);
 					}
 				},
 				UE::AnimNext::ETaskRunLocation::Before);
 		}
 		else
 		{
-			ModuleInstance->SetVariable(ResultOutputVariableName, Result);
+			ModuleInstance->SetVariable(ResultOutputVariableName, MotionMatchingState.SearchResult);
 		}
 	}
 
@@ -579,42 +459,40 @@ namespace UE::AnimNext
 #endif // WITH_EDITOR
 
 		// if it's not a multi character interaction we can skip the warping logic entirely
-		if (InstanceData->MotionMatchingInteractionState.SelectedAnim == nullptr)
+		const FPoseSearchBlueprintResult& SearchResult = InstanceData->MotionMatchingState.SearchResult;
+		if (SearchResult.bIsInteraction)
 		{
-			return;
-		}
-
-		const FPoseSearchBlueprintResult& CurrentResult = InstanceData->MotionMatchingInteractionState;
-		if (const UMultiAnimAsset* MultiAnimAsset = Cast<UMultiAnimAsset>(CurrentResult.SelectedAnim))
-		{
-			const int32 CurrentResultRoleIndex = GetRoleIndex(MultiAnimAsset, CurrentResult.Role);
-			if (CurrentResultRoleIndex != INDEX_NONE)
+			if (const UMultiAnimAsset* MultiAnimAsset = Cast<UMultiAnimAsset>(SearchResult.SelectedAnim))
 			{
-				// @todo: WIP hacky, non thread safe (unless proper tick dependencies are in place) way to retrieve the mesh transform until we find a better way
-				if (const FAnimNextModuleInstance* ModuleInstance = Context.GetRootGraphInstance().GetModuleInstance())
+				const int32 CurrentResultRoleIndex = GetRoleIndex(MultiAnimAsset, SearchResult.Role);
+				if (CurrentResultRoleIndex != INDEX_NONE)
 				{
-					if (const UAnimNextComponent* AnimNextComponent = Cast<UAnimNextComponent>(ModuleInstance->GetObject()))
+					// @todo: WIP hacky, non thread safe (unless proper tick dependencies are in place) way to retrieve the mesh transform until we find a better way
+					if (const FAnimNextModuleInstance* ModuleInstance = Context.GetRootGraphInstance().GetModuleInstance())
 					{
-						const AActor* Actor = AnimNextComponent->GetOwner();
-						check(Actor);
-						if (const USkeletalMeshComponent* SkeletalMeshComponent = Actor->GetComponentByClass<USkeletalMeshComponent>())
+						if (const UAnimNextComponent* AnimNextComponent = Cast<UAnimNextComponent>(ModuleInstance->GetObject()))
 						{
-							const FSharedData* SharedData = Binding.GetSharedData<FSharedData>();
-							check(SharedData);
-		
-							FAnimNextMotionMatchingTask Task;
-							Task.ComponentTransform = SkeletalMeshComponent->GetComponentTransform();
-							Task.InstanceData = InstanceData;
-							Task.CurrentResultRoleIndex = CurrentResultRoleIndex;
-							Task.bWarpUsingRootBone = SharedData->GetbWarpUsingRootBone(Binding);
-							Task.WarpingRotationRatio = SharedData->GetWarpingRotationRatio(Binding);
-							Task.WarpingTranslationRatio = SharedData->GetWarpingTranslationRatio(Binding);
+							const AActor* Actor = AnimNextComponent->GetOwner();
+							check(Actor);
+							if (const USkeletalMeshComponent* SkeletalMeshComponent = Actor->GetComponentByClass<USkeletalMeshComponent>())
+							{
+								const FSharedData* SharedData = Binding.GetSharedData<FSharedData>();
+								check(SharedData);
+
+								FAnimNextMotionMatchingTask Task;
+								Task.ComponentTransform = SkeletalMeshComponent->GetComponentTransform();
+								Task.InstanceData = InstanceData;
+								Task.CurrentResultRoleIndex = CurrentResultRoleIndex;
+								Task.bWarpUsingRootBone = SharedData->GetbWarpUsingRootBone(Binding);
+								Task.WarpingRotationRatio = SharedData->GetWarpingRotationRatio(Binding);
+								Task.WarpingTranslationRatio = SharedData->GetWarpingTranslationRatio(Binding);
 #if ENABLE_ANIM_DEBUG
-							// Debug Object for VisualLogger
-							Task.HostObject = Context.GetHostObject();
+								// Debug Object for VisualLogger
+								Task.HostObject = Context.GetHostObject();
 #endif // ENABLE_ANIM_DEBUG 
 
-							Context.AppendTask(Task);
+								Context.AppendTask(Task);
+							}
 						}
 					}
 				}
@@ -643,15 +521,17 @@ void FAnimNextMotionMatchingTask::Execute(UE::AnimNext::FEvaluationVM& VM) const
 		return;
 	}
 
-	const FPoseSearchBlueprintResult& CurrentResult = InstanceData->MotionMatchingInteractionState;
-	const UMultiAnimAsset* MultiAnimAsset = Cast<UMultiAnimAsset>(CurrentResult.SelectedAnim);
+	const FPoseSearchBlueprintResult& SearchResult = InstanceData->MotionMatchingState.SearchResult;
+	check(SearchResult.bIsInteraction);
+
+	const UMultiAnimAsset* MultiAnimAsset = Cast<UMultiAnimAsset>(SearchResult.SelectedAnim);
 	if (!MultiAnimAsset)
 	{
 		return;
 	}
 
 	const int32 NumRoles = MultiAnimAsset->GetNumRoles();
-	if (CurrentResult.ActorRootTransforms.Num() != NumRoles)
+	if (SearchResult.ActorRootTransforms.Num() != NumRoles)
 	{
 		// warping is supported only for UMultiAnimAsset(s)
 		return;
@@ -666,7 +546,7 @@ void FAnimNextMotionMatchingTask::Execute(UE::AnimNext::FEvaluationVM& VM) const
 	TArray<FTransform, TInlineAllocator<PreallocatedRolesNum>> FullAlignedTransforms;
 	FullAlignedTransforms.SetNum(NumRoles);
 
-	CalculateFullAlignedTransforms(CurrentResult, bWarpUsingRootBone, FullAlignedTransforms);
+	CalculateFullAlignedTransforms(SearchResult, bWarpUsingRootBone, FullAlignedTransforms);
 
 	// @todo: implement this properly once we have UE::AnimationWarping::FRootOffsetProvider
 	const FTransform& MeshWithoutOffset = ComponentTransform;
@@ -686,7 +566,7 @@ void FAnimNextMotionMatchingTask::Execute(UE::AnimNext::FEvaluationVM& VM) const
 
 		for (int32 Index = 0; Index < NumRoles; ++Index)
 		{
-			const FTransform& ActorRootTransform = CurrentResult.ActorRootTransforms[Index];
+			const FTransform& ActorRootTransform = SearchResult.ActorRootTransforms[Index];
 			const FTransform& FullAlignedTransform = FullAlignedTransforms[Index];
 
 			UE_VLOG_SEGMENT_THICK(HostObject, LogName, Display, FullAlignedTransform.GetLocation(), ActorRootTransform.GetLocation(), FColorList::Orange, 1.f, TEXT(""));
diff --git a/Engine/Plugins/Experimental/AnimNextPoseSearch/Source/AnimNextPoseSearch/Private/MotionMatchingTrait.h b/Engine/Plugins/Experimental/AnimNextPoseSearch/Source/AnimNextPoseSearch/Private/MotionMatchingTrait.h
--- a/Engine/Plugins/Experimental/AnimNextPoseSearch/Source/AnimNextPoseSearch/Private/MotionMatchingTrait.h
+++ b/Engine/Plugins/Experimental/AnimNextPoseSearch/Source/AnimNextPoseSearch/Private/MotionMatchingTrait.h
@@ -143,9 +143,6 @@ struct FMotionMatchingTrait : FAdditiveTrait, IUpdate, IEvaluate
 	{
 		FMotionMatchingState MotionMatchingState;
 
-		// @todo: merge this with MotionMatchingState
-		FPoseSearchBlueprintResult MotionMatchingInteractionState;
-
 #if WITH_EDITOR
 		bool bIsPostEvaluateBeingCalled = true;
 #endif // WITH_EDITOR
diff --git a/Engine/Source/Runtime/Engine/Public/Animation/AnimPoseSearchProvider.h b/Engine/Source/Runtime/Engine/Public/Animation/AnimPoseSearchProvider.h
--- a/Engine/Source/Runtime/Engine/Public/Animation/AnimPoseSearchProvider.h
+++ b/Engine/Source/Runtime/Engine/Public/Animation/AnimPoseSearchProvider.h
@@ -30,12 +30,15 @@ public:
 		bool bIsFromContinuingPlaying = false;
 		bool bMirrored = false;
 		float WantedPlayRate = 1.f;
+		FVector BlendParameters = FVector::ZeroVector;
 	};
 
 	struct FSearchPlayingAsset
 	{
 		const UObject* Asset = nullptr;
 		float AccumulatedTime = 0.f;
+		bool bMirrored = false;
+		FVector BlendParameters = FVector::ZeroVector;
 	};
 
 	struct FSearchFutureAsset : FSearchPlayingAsset
@@ -57,17 +60,6 @@ public:
 	* @return	FSearchResult					The search result identifying the asset from AssetsToSearch or PlayingAsset that most closely matches the query
 	*/
 	virtual FSearchResult Search(const FAnimationBaseContext& GraphContext, TArrayView<const UObject*> AssetsToSearch, const FSearchPlayingAsset& PlayingAsset, const FSearchFutureAsset& FutureAsset) const = 0;
-	
-	UE_DEPRECATED(5.4, "Use the Search with FutureAsset instead")
-	virtual FSearchResult Search(const FAnimationBaseContext& GraphContext, TArrayView<const UObject*> AssetsToSearch,
-		const UObject* PlayingAsset = nullptr, float PlayingAssetAccumulatedTime = 0.f) const
-	{
-		FSearchPlayingAsset SearchPlayingAsset;
-		SearchPlayingAsset.Asset = PlayingAsset;
-		SearchPlayingAsset.AccumulatedTime = PlayingAssetAccumulatedTime;
-
-		return Search(GraphContext, AssetsToSearch, SearchPlayingAsset, FSearchFutureAsset());
-	}
 };
 
 } // namespace UE::Anim
\ No newline at end of file
