Fix deprecation warnings

diff --git a/Engine/Plugins/Experimental/AnimNextPoseSearch/AnimNextPoseSearch.uplugin b/Engine/Plugins/Experimental/AnimNextPoseSearch/AnimNextPoseSearch.uplugin
--- a/Engine/Plugins/Experimental/AnimNextPoseSearch/AnimNextPoseSearch.uplugin
+++ b/Engine/Plugins/Experimental/AnimNextPoseSearch/AnimNextPoseSearch.uplugin
@@ -28,6 +28,10 @@
 		{
 			"Name": "PoseSearch",
 			"Enabled": true
+		},
+		{
+			"Name": "Chooser",
+			"Enabled": true
 		}
 	],
 	"Modules" :
diff --git a/Engine/Plugins/Experimental/AnimNextPoseSearch/Source/AnimNextPoseSearch/AnimNextPoseSearch.Build.cs b/Engine/Plugins/Experimental/AnimNextPoseSearch/Source/AnimNextPoseSearch/AnimNextPoseSearch.Build.cs
--- a/Engine/Plugins/Experimental/AnimNextPoseSearch/Source/AnimNextPoseSearch/AnimNextPoseSearch.Build.cs
+++ b/Engine/Plugins/Experimental/AnimNextPoseSearch/Source/AnimNextPoseSearch/AnimNextPoseSearch.Build.cs
@@ -17,6 +17,7 @@ namespace UnrealBuildTool.Rules
 					"HierarchyTableRuntime",
 					"HierarchyTableAnimationRuntime",
 					"PoseSearch",
+					"Chooser",
 					"RigVM",
 				}
 			);
diff --git a/Engine/Plugins/Experimental/AnimNextPoseSearch/Source/AnimNextPoseSearch/Private/MotionMatchingTrait.cpp b/Engine/Plugins/Experimental/AnimNextPoseSearch/Source/AnimNextPoseSearch/Private/MotionMatchingTrait.cpp
--- a/Engine/Plugins/Experimental/AnimNextPoseSearch/Source/AnimNextPoseSearch/Private/MotionMatchingTrait.cpp
+++ b/Engine/Plugins/Experimental/AnimNextPoseSearch/Source/AnimNextPoseSearch/Private/MotionMatchingTrait.cpp
@@ -236,8 +236,9 @@ namespace UE::AnimNext
 				PoseSearchPlayRate = bPoseSearchPlayRateEnabled ? PoseSearchPlayRate : FFloatInterval(1.f, 1.f);
 #endif // ENABLE_ANIM_DEBUG
 
+				FChooserEvaluationContext EvaluationContext(const_cast<UObject*>(AnimContext));
 				UPoseSearchLibrary::UpdateMotionMatchingState(
-					AnimContext,
+					&EvaluationContext,
 					PoseHistory,
 					DatabasesToSearch,
 					DeltaTime,
