Added method to query the currently selected pose search database and its metadata tags in AnimNext

#rnx

diff --git a/Engine/Plugins/Experimental/AnimNextPoseSearch/Source/AnimNextPoseSearch/Private/RigUnit_PoseSearchDatabaseGetTags.cpp b/Engine/Plugins/Experimental/AnimNextPoseSearch/Source/AnimNextPoseSearch/Private/RigUnit_PoseSearchDatabaseGetTags.cpp
new file mode 100644
--- /dev/null
+++ b/Engine/Plugins/Experimental/AnimNextPoseSearch/Source/AnimNextPoseSearch/Private/RigUnit_PoseSearchDatabaseGetTags.cpp
@@ -0,0 +1,14 @@
+// Copyright Epic Games, Inc. All Rights Reserved.
+
+#include "RigUnit_PoseSearchDatabaseGetTags.h"
+#include "PoseSearch/PoseSearchLibrary.h"
+
+#include UE_INLINE_GENERATED_CPP_BY_NAME(RigUnit_PoseSearchDatabaseGetTags)
+
+FRigUnit_PoseSearchDatabaseGetTags_Execute()
+{
+	if (Database)
+	{
+		UPoseSearchLibrary::GetDatabaseTags(Database, OutTags);
+	}
+}
\ No newline at end of file
diff --git a/Engine/Plugins/Experimental/AnimNextPoseSearch/Source/AnimNextPoseSearch/Private/RigUnit_PoseSearchDatabaseGetTags.h b/Engine/Plugins/Experimental/AnimNextPoseSearch/Source/AnimNextPoseSearch/Private/RigUnit_PoseSearchDatabaseGetTags.h
new file mode 100644
--- /dev/null
+++ b/Engine/Plugins/Experimental/AnimNextPoseSearch/Source/AnimNextPoseSearch/Private/RigUnit_PoseSearchDatabaseGetTags.h
@@ -0,0 +1,24 @@
+// Copyright Epic Games, Inc. All Rights Reserved.
+
+#pragma once
+
+#include "Graph/RigUnit_AnimNextBase.h"
+#include "RigUnit_PoseSearchDatabaseGetTags.generated.h"
+
+class UPoseSearchDatabase;
+
+/** Get metadata tags for a specified pose search database */
+USTRUCT(meta=(DisplayName="Get Database Tags", Category="Pose Search", NodeColor="0, 1, 1", Keywords="PoseSearch"))
+struct FRigUnit_PoseSearchDatabaseGetTags : public FRigUnit_AnimNextBase
+{
+	GENERATED_BODY()
+
+	RIGVM_METHOD()
+	virtual void Execute() override;
+
+	UPROPERTY(EditAnywhere, Transient, Category = "Pose Search", meta = (Input))
+	TObjectPtr<UPoseSearchDatabase> Database;
+
+	UPROPERTY(EditAnywhere, Transient, Category = "Pose Search", meta = (Output))
+	TArray<FName> OutTags;
+};
diff --git a/Engine/Plugins/Experimental/AnimNextPoseSearch/Source/AnimNextPoseSearch/Private/RigUnit_PoseSearchResultGetSelectedDatabase.cpp b/Engine/Plugins/Experimental/AnimNextPoseSearch/Source/AnimNextPoseSearch/Private/RigUnit_PoseSearchResultGetSelectedDatabase.cpp
new file mode 100644
--- /dev/null
+++ b/Engine/Plugins/Experimental/AnimNextPoseSearch/Source/AnimNextPoseSearch/Private/RigUnit_PoseSearchResultGetSelectedDatabase.cpp
@@ -0,0 +1,14 @@
+// Copyright Epic Games, Inc. All Rights Reserved.
+
+#include "RigUnit_PoseSearchResultGetSelectedDatabase.h"
+#include "PoseSearch/PoseSearchDatabase.h"
+
+#include UE_INLINE_GENERATED_CPP_BY_NAME(RigUnit_PoseSearchResultGetSelectedDatabase)
+
+FRigUnit_PoseSearchResultGetSelectedDatabase_Execute()
+{
+	if (PoseSearchResult.SelectedDatabase.IsValid())
+	{
+		OutDatabase = const_cast<UPoseSearchDatabase*>(PoseSearchResult.SelectedDatabase.Get());
+	}
+}
\ No newline at end of file
diff --git a/Engine/Plugins/Experimental/AnimNextPoseSearch/Source/AnimNextPoseSearch/Private/RigUnit_PoseSearchResultGetSelectedDatabase.h b/Engine/Plugins/Experimental/AnimNextPoseSearch/Source/AnimNextPoseSearch/Private/RigUnit_PoseSearchResultGetSelectedDatabase.h
new file mode 100644
--- /dev/null
+++ b/Engine/Plugins/Experimental/AnimNextPoseSearch/Source/AnimNextPoseSearch/Private/RigUnit_PoseSearchResultGetSelectedDatabase.h
@@ -0,0 +1,26 @@
+// Copyright Epic Games, Inc. All Rights Reserved.
+
+#pragma once
+
+#include "Graph/RigUnit_AnimNextBase.h"
+#include "PoseSearch/PoseSearchResult.h"
+
+#include "RigUnit_PoseSearchResultGetSelectedDatabase.generated.h"
+
+/** Gets selected database from a pose search result */
+USTRUCT(meta=(DisplayName="Get Selected Database ", Category="Pose Search", NodeColor="0, 1, 1", Keywords="PoseSearch"))
+struct FRigUnit_PoseSearchResultGetSelectedDatabase : public FRigUnit_AnimNextBase
+{
+	GENERATED_BODY()
+
+	RIGVM_METHOD()
+	virtual void Execute() override;
+
+	// Role to search for within the MultiAnimAsset
+	UPROPERTY(EditAnywhere, Transient, Category = "Pose Search", meta = (Input))
+	FPoseSearchBlueprintResult PoseSearchResult;
+
+	// Result of the role search in the MultiAnimAsset
+	UPROPERTY(EditAnywhere, Transient, Category = "Pose Search", meta = (Output))
+	TObjectPtr<UPoseSearchDatabase> OutDatabase = nullptr;
+};
\ No newline at end of file
