PoseSearch - MotionMatchingTrait, keep interaction alive

diff --git a/Engine/Plugins/Animation/PoseSearch/Source/Runtime/Private/PoseSearchContext.cpp b/Engine/Plugins/Animation/PoseSearch/Source/Runtime/Private/PoseSearchContext.cpp
--- a/Engine/Plugins/Animation/PoseSearch/Source/Runtime/Private/PoseSearchContext.cpp
+++ b/Engine/Plugins/Animation/PoseSearch/Source/Runtime/Private/PoseSearchContext.cpp
@@ -1035,7 +1035,7 @@ FTransform FSearchContext::GetWorldBoneTransformAtTime(float SampleTime, const F
 				{
 					if (!PoseHistory->IsEmpty())
 					{
-						UE_LOG(LogPoseSearch, Warning, TEXT("FSearchContext::GetWorldBoneTransformAtTime - Couldn't find BoneIndexType %d (%s) requested by %s"), BoneIndexType, *Skeleton->GetReferenceSkeleton().GetBoneName(BoneIndexType).ToString(), *Schema->GetName());
+						UE_LOG(LogPoseSearch, Warning, TEXT("FSearchContext::GetWorldBoneTransformAtTime - Couldn't find BoneIndexType %d (%s) for Skeleton %s in the input IPoseHistory requested by %s. Consider adding it to the Pose History!"), BoneIndexType, *Skeleton->GetReferenceSkeleton().GetBoneName(BoneIndexType).ToString(), *Skeleton->GetName(), *Schema->GetName());
 					}
 				}
 				else
diff --git a/Engine/Plugins/Experimental/AnimNextPoseSearch/Source/AnimNextPoseSearch/Private/MotionMatchingTrait.cpp b/Engine/Plugins/Experimental/AnimNextPoseSearch/Source/AnimNextPoseSearch/Private/MotionMatchingTrait.cpp
--- a/Engine/Plugins/Experimental/AnimNextPoseSearch/Source/AnimNextPoseSearch/Private/MotionMatchingTrait.cpp
+++ b/Engine/Plugins/Experimental/AnimNextPoseSearch/Source/AnimNextPoseSearch/Private/MotionMatchingTrait.cpp
@@ -205,17 +205,40 @@ namespace UE::AnimNext
 		// performing the regular single character motion matching search in case there's no MM interaction
 		if (!bIsInteraction)
 		{
-			if (MotionMatchingState.SearchResult.bIsInteraction)
+			const bool bKeepInteractionAlive = SharedData->GetbKeepInteractionAlive(Binding);
+			if (!MotionMatchingState.SearchResult.bIsInteraction)
 			{
-				MotionMatchingState.SearchResult = FPoseSearchBlueprintResult();
+				// @todo: ask blendstack what it's playing instead of relying on MotionMatchingInteractionState.SelectedAnim
+				//MotionMatchingState.SearchResult.SelectedAnim = GetAnimAsset();
+				MotionMatchingState.SearchResult.SelectedTime = TimelineNormalizedTime;
+				// @todo: ask blendstack what the blend parameters are
+				//MotionMatchingState.SearchResult.BlendParameters = GetBlendParameters();
 			}
-			else
+			else if (bKeepInteractionAlive)
 			{
 				// @todo: ask blendstack what it's playing instead of relying on MotionMatchingInteractionState.SelectedAnim
 				//MotionMatchingState.SearchResult.SelectedAnim = GetAnimAsset();
 				MotionMatchingState.SearchResult.SelectedTime = TimelineNormalizedTime;
 				// @todo: ask blendstack what the blend parameters are
 				//MotionMatchingState.SearchResult.BlendParameters = GetBlendParameters();
+
+				// checking if the kept alive interaction has reached the end of animation 
+				if (MotionMatchingState.SearchResult.SelectedDatabase.IsValid() &&
+					INDEX_NONE == MotionMatchingState.SearchResult.SelectedDatabase->GetPoseIndex(MotionMatchingState.SearchResult.SelectedAnim.Get(), MotionMatchingState.SearchResult.SelectedTime, MotionMatchingState.SearchResult.bIsMirrored, MotionMatchingState.SearchResult.BlendParameters))
+				{
+					MotionMatchingState.SearchResult = FPoseSearchBlueprintResult();
+				}
+				else
+				{
+					// we're keeping alive only the animation part of the search result
+					MotionMatchingState.SearchResult.ActorRootTransforms.Reset();
+					MotionMatchingState.SearchResult.ActorRootBoneTransforms.Reset();
+					MotionMatchingState.SearchResult.AnimContexts.Reset();
+				}
+			}
+			else
+			{
+				MotionMatchingState.SearchResult = FPoseSearchBlueprintResult();
 			}
 
 			const TConstArrayView<TObjectPtr<const UPoseSearchDatabase>> DatabasesToSearch = SharedData->GetDatabases(Binding);
diff --git a/Engine/Plugins/Experimental/AnimNextPoseSearch/Source/AnimNextPoseSearch/Private/MotionMatchingTrait.h b/Engine/Plugins/Experimental/AnimNextPoseSearch/Source/AnimNextPoseSearch/Private/MotionMatchingTrait.h
--- a/Engine/Plugins/Experimental/AnimNextPoseSearch/Source/AnimNextPoseSearch/Private/MotionMatchingTrait.h
+++ b/Engine/Plugins/Experimental/AnimNextPoseSearch/Source/AnimNextPoseSearch/Private/MotionMatchingTrait.h
@@ -92,6 +92,9 @@ struct FMotionMatchingTraitSharedData : public FAnimNextTraitSharedData
 	UPROPERTY(EditAnywhere, Category = Interaction, meta = (PinHiddenByDefault))
 	bool bValidateResultAgainstAvailabilities = true;
 
+	UPROPERTY(EditAnywhere, Category = Interaction, meta = (PinHiddenByDefault))
+	bool bKeepInteractionAlive = false;
+
 	UPROPERTY(EditAnywhere, Category = Interaction, meta = (PinHiddenByDefault))
 	bool bWarpUsingRootBone = true;
 
@@ -122,6 +125,7 @@ struct FMotionMatchingTraitSharedData : public FAnimNextTraitSharedData
 		GeneratorMacro(MotionMatchingResultModuleHandle) \
 		GeneratorMacro(Availabilities) \
 		GeneratorMacro(bValidateResultAgainstAvailabilities) \
+		GeneratorMacro(bKeepInteractionAlive) \
 		GeneratorMacro(bWarpUsingRootBone) \
 		GeneratorMacro(WarpingRotationRatio) \
 		GeneratorMacro(WarpingTranslationRatio) \
