PoseSearch - converting FPoseSearchBlueprintResult TWeakObjectPtr to TObjectPtr

diff --git a/Engine/Plugins/Animation/PoseSearch/Source/Runtime/Private/AnimNode_MotionMatching.cpp b/Engine/Plugins/Animation/PoseSearch/Source/Runtime/Private/AnimNode_MotionMatching.cpp
--- a/Engine/Plugins/Animation/PoseSearch/Source/Runtime/Private/AnimNode_MotionMatching.cpp
+++ b/Engine/Plugins/Animation/PoseSearch/Source/Runtime/Private/AnimNode_MotionMatching.cpp
@@ -181,7 +181,7 @@ void FAnimNode_MotionMatching::UpdateAssetPlayer(const FAnimationUpdateContext&
 		, EventToSearch
 	);
 
-	const bool bJumpToPose = MotionMatchingState.SearchResult.SelectedAnim.IsValid() && !MotionMatchingState.SearchResult.bIsContinuingPoseSearch;
+	const bool bJumpToPose = MotionMatchingState.SearchResult.SelectedAnim && !MotionMatchingState.SearchResult.bIsContinuingPoseSearch;
 
 	PRAGMA_DISABLE_DEPRECATION_WARNINGS
 	MotionMatchingState.bJumpedToPose = bJumpToPose;
diff --git a/Engine/Plugins/Animation/PoseSearch/Source/Runtime/Private/AnimNode_MotionMatchingInteraction.cpp b/Engine/Plugins/Animation/PoseSearch/Source/Runtime/Private/AnimNode_MotionMatchingInteraction.cpp
--- a/Engine/Plugins/Animation/PoseSearch/Source/Runtime/Private/AnimNode_MotionMatchingInteraction.cpp
+++ b/Engine/Plugins/Animation/PoseSearch/Source/Runtime/Private/AnimNode_MotionMatchingInteraction.cpp
@@ -65,7 +65,7 @@ void FAnimNode_MotionMatchingInteraction::UpdateAssetPlayer(const FAnimationUpda
 	{
 		// @todo: retrieve the playing UMultiAnimAsset from the UAnimAsset blendspace is playing via GetAnimAsset() and blend parameters via GetBlendParameters() and mirror state
 		FPoseSearchContinuingProperties ContinuingProperties;
-		ContinuingProperties.PlayingAsset = CurrentResult.SelectedAnim.Get();
+		ContinuingProperties.PlayingAsset = CurrentResult.SelectedAnim;
 		ContinuingProperties.PlayingAssetAccumulatedTime = GetAccumulatedTime();
 		ContinuingProperties.bIsPlayingAssetMirrored = CurrentResult.bIsMirrored;
 		ContinuingProperties.PlayingAssetBlendParameters = CurrentResult.BlendParameters;
@@ -98,7 +98,7 @@ void FAnimNode_MotionMatchingInteraction::UpdateAssetPlayer(const FAnimationUpda
 	const bool bDidBlendToRequestAnInertialBlend = bBlendToExecuted && bUseInertialBlend;
 	UE::Anim::TOptionalScopedGraphMessage<UE::Anim::FAnimInertializationSyncScope> InertializationSync(bDidBlendToRequestAnInertialBlend, Context);
 	
-	if (CurrentResult.SelectedAnim.IsValid())
+	if (CurrentResult.SelectedAnim)
 	{
 		UpdatePlayRate(CurrentResult.WantedPlayRate);
 		UpdateBlendspaceParameters(BlendspaceUpdateMode, CurrentResult.BlendParameters);
diff --git a/Engine/Plugins/Animation/PoseSearch/Source/Runtime/Private/PoseSearchInteractionIsland.cpp b/Engine/Plugins/Animation/PoseSearch/Source/Runtime/Private/PoseSearchInteractionIsland.cpp
--- a/Engine/Plugins/Animation/PoseSearch/Source/Runtime/Private/PoseSearchInteractionIsland.cpp
+++ b/Engine/Plugins/Animation/PoseSearch/Source/Runtime/Private/PoseSearchInteractionIsland.cpp
@@ -1135,7 +1135,7 @@ bool FInteractionIsland::GetResult_AnyThread(const UObject* AnimContext, FPoseSe
 						{
 							Result.ActorRootTransforms[MultiAnimAssetRoleIndex] = SearchResult.ActorRootTransforms[*InteractionSearchContextRoleIndex];
 							Result.ActorRootBoneTransforms[MultiAnimAssetRoleIndex] = SearchResult.ActorRootBoneTransforms[*InteractionSearchContextRoleIndex];
-							Result.AnimContexts[MultiAnimAssetRoleIndex] = SearchContext.AnimContexts[*InteractionSearchContextRoleIndex];
+							Result.AnimContexts[MultiAnimAssetRoleIndex] = SearchContext.AnimContexts[*InteractionSearchContextRoleIndex].Get();
 						}
 						else
 						{
@@ -1152,7 +1152,9 @@ bool FInteractionIsland::GetResult_AnyThread(const UObject* AnimContext, FPoseSe
 
 					Result.ActorRootTransforms = SearchResult.ActorRootTransforms;
 					Result.ActorRootBoneTransforms = SearchResult.ActorRootBoneTransforms;
-					Result.AnimContexts = SearchContext.AnimContexts;
+					
+					Result.AnimContexts.SetNum(1);
+					Result.AnimContexts[0] = SearchContext.AnimContexts[0].Get();
 				}
 
 				// we found our AnimContext: we can stop searching
diff --git a/Engine/Plugins/Animation/PoseSearch/Source/Runtime/Private/PoseSearchInteractionSubsystem.cpp b/Engine/Plugins/Animation/PoseSearch/Source/Runtime/Private/PoseSearchInteractionSubsystem.cpp
--- a/Engine/Plugins/Animation/PoseSearch/Source/Runtime/Private/PoseSearchInteractionSubsystem.cpp
+++ b/Engine/Plugins/Animation/PoseSearch/Source/Runtime/Private/PoseSearchInteractionSubsystem.cpp
@@ -1193,7 +1193,7 @@ void UPoseSearchInteractionSubsystem::Query_AnyThread(const TArrayView<const FPo
 	{
 		Island->DoSearch_AnyThread(AnimContext, ContinuingProperties, Result);
 
-		if (bValidateResultAgainstAvailabilities && Result.SelectedAnim != nullptr)
+		if (bValidateResultAgainstAvailabilities && Result.SelectedAnim)
 		{
 			bool bResultValidated = false;
 
diff --git a/Engine/Plugins/Animation/PoseSearch/Source/Runtime/Private/PoseSearchResult.cpp b/Engine/Plugins/Animation/PoseSearch/Source/Runtime/Private/PoseSearchResult.cpp
--- a/Engine/Plugins/Animation/PoseSearch/Source/Runtime/Private/PoseSearchResult.cpp
+++ b/Engine/Plugins/Animation/PoseSearch/Source/Runtime/Private/PoseSearchResult.cpp
@@ -193,7 +193,7 @@ const UAnimationAsset* FSearchResult::GetCurrentResultAnimationAsset(const FRole
 void FSearchResult::InitFrom(const FPoseSearchBlueprintResult& BlueprintResult)
 {
 	PoseCost = FPoseSearchCost(BlueprintResult.SearchCost, 0.f, 0.f);
-	PoseIdx = BlueprintResult.SelectedDatabase.IsValid() ? BlueprintResult.SelectedDatabase->GetPoseIndex(BlueprintResult.SelectedAnim.Get(), BlueprintResult.SelectedTime, BlueprintResult.bIsMirrored, BlueprintResult.BlendParameters) : INDEX_NONE;
+	PoseIdx = BlueprintResult.SelectedDatabase ? BlueprintResult.SelectedDatabase->GetPoseIndex(BlueprintResult.SelectedAnim.Get(), BlueprintResult.SelectedTime, BlueprintResult.bIsMirrored, BlueprintResult.BlendParameters) : INDEX_NONE;
 	EventPoseIdx = INDEX_NONE;
 	Database = BlueprintResult.SelectedDatabase;
 	AssetTime = BlueprintResult.SelectedTime;
diff --git a/Engine/Plugins/Animation/PoseSearch/Source/Runtime/Public/PoseSearch/AnimNode_MotionMatching.h b/Engine/Plugins/Animation/PoseSearch/Source/Runtime/Public/PoseSearch/AnimNode_MotionMatching.h
--- a/Engine/Plugins/Animation/PoseSearch/Source/Runtime/Public/PoseSearch/AnimNode_MotionMatching.h
+++ b/Engine/Plugins/Animation/PoseSearch/Source/Runtime/Public/PoseSearch/AnimNode_MotionMatching.h
@@ -160,6 +160,7 @@ private:
 	FPoseSearchEvent EventToSearch;
 
 	// Encapsulated motion matching algorithm and internal state
+	UPROPERTY(Transient)
 	FMotionMatchingState MotionMatchingState;
 
 	// Update Counter for detecting being relevant
diff --git a/Engine/Plugins/Animation/PoseSearch/Source/Runtime/Public/PoseSearch/AnimNode_MotionMatchingInteraction.h b/Engine/Plugins/Animation/PoseSearch/Source/Runtime/Public/PoseSearch/AnimNode_MotionMatchingInteraction.h
--- a/Engine/Plugins/Animation/PoseSearch/Source/Runtime/Public/PoseSearch/AnimNode_MotionMatchingInteraction.h
+++ b/Engine/Plugins/Animation/PoseSearch/Source/Runtime/Public/PoseSearch/AnimNode_MotionMatchingInteraction.h
@@ -64,9 +64,10 @@ struct POSESEARCH_API FAnimNode_MotionMatchingInteraction : public FAnimNode_Ble
 
 	bool IsInteracting() const { return CurrentResult.bIsInteraction; }
 
-protected:
+private:
 	bool NeedsReset(const FAnimationUpdateContext& Context) const;
 
+	UPROPERTY(Transient)
 	FPoseSearchBlueprintResult CurrentResult;
 	FTransform MeshWithOffset = FTransform::Identity;
 	FTransform MeshWithoutOffset = FTransform::Identity;
diff --git a/Engine/Plugins/Animation/PoseSearch/Source/Runtime/Public/PoseSearch/PoseSearchLibrary.h b/Engine/Plugins/Animation/PoseSearch/Source/Runtime/Public/PoseSearch/PoseSearchLibrary.h
--- a/Engine/Plugins/Animation/PoseSearch/Source/Runtime/Public/PoseSearch/PoseSearchLibrary.h
+++ b/Engine/Plugins/Animation/PoseSearch/Source/Runtime/Public/PoseSearch/PoseSearchLibrary.h
@@ -44,10 +44,20 @@ enum class EPoseSearchInterruptMode : uint8
 	ForceInterruptAndInvalidateContinuingPose,
 };
 
-PRAGMA_DISABLE_DEPRECATION_WARNINGS
+USTRUCT()
 struct FMotionMatchingState
 {
+	GENERATED_BODY()
+
+PRAGMA_DISABLE_DEPRECATION_WARNINGS
+	FMotionMatchingState() = default;
+	FMotionMatchingState(const FMotionMatchingState& Other) = default;
+	FMotionMatchingState(FMotionMatchingState&& Other) = default;
+	FMotionMatchingState& operator=(const FMotionMatchingState& Other) = default;
+	FMotionMatchingState& operator=(FMotionMatchingState&& Other) = default;
+	~FMotionMatchingState() = default;
 PRAGMA_ENABLE_DEPRECATION_WARNINGS
+
 	// Reset the state to a default state using the current Database
 	UE_DEPRECATED(5.6, "Use other Reset signatures instead")
 	void Reset(const FTransform& ComponentTransform);
@@ -63,10 +73,10 @@ PRAGMA_ENABLE_DEPRECATION_WARNINGS
 
 	FVector GetEstimatedFutureRootMotionVelocity() const;
 
-	//@todo: fix all references to CurrentSearchResult before deprecating it
-	//UE_DEPRECATED(5.6, "Use SearchResult instead")
+	UE_DEPRECATED(5.6, "Use SearchResult instead")
 	UE::PoseSearch::FSearchResult CurrentSearchResult;
 
+	UPROPERTY(Transient)
 	FPoseSearchBlueprintResult SearchResult;
 
 	// Time since the last pose jump
diff --git a/Engine/Plugins/Animation/PoseSearch/Source/Runtime/Public/PoseSearch/PoseSearchResult.h b/Engine/Plugins/Animation/PoseSearch/Source/Runtime/Public/PoseSearch/PoseSearchResult.h
--- a/Engine/Plugins/Animation/PoseSearch/Source/Runtime/Public/PoseSearch/PoseSearchResult.h
+++ b/Engine/Plugins/Animation/PoseSearch/Source/Runtime/Public/PoseSearch/PoseSearchResult.h
@@ -119,7 +119,6 @@ USTRUCT(BlueprintType, Category="Animation|Pose Search")
 struct POSESEARCH_API FPoseSearchBlueprintResult
 {
 	GENERATED_BODY()
-public:
 
 #if WITH_EDITORONLY_DATA
 	// Deprecated in 5.6. Use "SelectedAnim" property instead
@@ -129,7 +128,7 @@ public:
 
 	// selected animation selected by the motion matching alghoritm that matches the query (trajectory + historical pose) the best
 	UPROPERTY(Transient, VisibleAnywhere, BlueprintReadOnly, Category=State)
-	TWeakObjectPtr<UObject> SelectedAnim = nullptr;
+	TObjectPtr<UObject> SelectedAnim = nullptr;
 	
 	// SelectedAnim associated time (to have the best result out of the motion matchign selection, play the SelectedAnim at the SelectedTime)
 	UPROPERTY(Transient, VisibleAnywhere, BlueprintReadOnly, Category=State)
@@ -157,7 +156,7 @@ public:
 
 	// database containig the SelectedAnim
 	UPROPERTY(Transient, VisibleAnywhere, BlueprintReadOnly, Category=State)
-	TWeakObjectPtr<const UPoseSearchDatabase> SelectedDatabase = nullptr;
+	TObjectPtr<const UPoseSearchDatabase> SelectedDatabase = nullptr;
 
 	// associated motion matching search cost for this result (how bad the SelectedAnim at SelectedTime etc looks in comparison to the query. The bigger the number, the more different the two are)
 	UPROPERTY(Transient, VisibleAnywhere, BlueprintReadOnly, Category=State)
@@ -182,7 +181,8 @@ public:
 	// Experimental, this feature might be removed without warning, not for production use
 	// AnimContexts involved in this interaction for all the roles in SelectedAnim (as UMultiAnimAsset). If the role is missing because optional the AnimContext will be null
 	// BE CAREFUL HOW YOU USE THEM! since accessing data from other actors without proper tick dependencies leads to threading problems!
-	TArray<TWeakObjectPtr<const UObject>> AnimContexts;
+	UPROPERTY(Transient)
+	TArray<TObjectPtr<const UObject>> AnimContexts;
 
 	// Experimental, this feature might be removed without warning, not for production use
 	bool InitFrom(const UE::PoseSearch::FSearchResult& SearchResult, float InWantedPlayRate);
diff --git a/Engine/Plugins/Experimental/AnimNextPoseSearch/Source/AnimNextPoseSearch/Private/MotionMatchingTrait.cpp b/Engine/Plugins/Experimental/AnimNextPoseSearch/Source/AnimNextPoseSearch/Private/MotionMatchingTrait.cpp
--- a/Engine/Plugins/Experimental/AnimNextPoseSearch/Source/AnimNextPoseSearch/Private/MotionMatchingTrait.cpp
+++ b/Engine/Plugins/Experimental/AnimNextPoseSearch/Source/AnimNextPoseSearch/Private/MotionMatchingTrait.cpp
@@ -185,7 +185,7 @@ namespace UE::AnimNext
 			if (MotionMatchingState.SearchResult.bIsInteraction)
 			{
 				// @todo: retrieve the playing UMultiAnimAsset from the UAnimAsset blendspace is playing via GetAnimAsset() and blend parameters via GetBlendParameters() and mirror state
-				ContinuingProperties.PlayingAsset = MotionMatchingState.SearchResult.SelectedAnim.Get();
+				ContinuingProperties.PlayingAsset = MotionMatchingState.SearchResult.SelectedAnim;
 				ContinuingProperties.PlayingAssetAccumulatedTime = TimelineNormalizedTime;
 				ContinuingProperties.bIsPlayingAssetMirrored = MotionMatchingState.SearchResult.bIsMirrored;
 				ContinuingProperties.PlayingAssetBlendParameters = MotionMatchingState.SearchResult.BlendParameters;
@@ -195,7 +195,7 @@ namespace UE::AnimNext
 			UPoseSearchInteractionLibrary::MotionMatchInteraction(SearchResult, Availabilities, AnimContext, ContinuingProperties, FName(), PoseHistory, SharedData->GetbValidateResultAgainstAvailabilities(Binding));
 			check(SearchResult.ActorRootTransforms.Num() == SearchResult.ActorRootBoneTransforms.Num());
 
-			if (SearchResult.SelectedAnim.IsValid())
+			if (SearchResult.SelectedAnim)
 			{
 				bIsInteraction = true;
 				MotionMatchingState.SearchResult = SearchResult;
@@ -223,7 +223,7 @@ namespace UE::AnimNext
 				//MotionMatchingState.SearchResult.BlendParameters = GetBlendParameters();
 
 				// checking if the kept alive interaction has reached the end of animation 
-				if (MotionMatchingState.SearchResult.SelectedDatabase.IsValid() &&
+				if (MotionMatchingState.SearchResult.SelectedDatabase &&
 					INDEX_NONE == MotionMatchingState.SearchResult.SelectedDatabase->GetPoseIndex(MotionMatchingState.SearchResult.SelectedAnim.Get(), MotionMatchingState.SearchResult.SelectedTime, MotionMatchingState.SearchResult.bIsMirrored, MotionMatchingState.SearchResult.BlendParameters))
 				{
 					MotionMatchingState.SearchResult = FPoseSearchBlueprintResult();
@@ -286,7 +286,7 @@ namespace UE::AnimNext
 		DesiredPlayRate = GVarAnimNextMotionMatchingPlayRateEnabled == UE::Private::EPlayRateState::Enabled ? DesiredPlayRate : 1.0f;
 #endif // ENABLE_ANIM_DEBUG
 
-		const bool bJumpToPose = MotionMatchingState.SearchResult.SelectedAnim.IsValid() && !MotionMatchingState.SearchResult.bIsContinuingPoseSearch;
+		const bool bJumpToPose = MotionMatchingState.SearchResult.SelectedAnim && !MotionMatchingState.SearchResult.bIsContinuingPoseSearch;
 		if (bJumpToPose)
 		{
 			if (UAnimationAsset* AnimationAsset = MotionMatchingState.SearchResult.GetAnimationAssetForRole())
@@ -523,6 +523,24 @@ namespace UE::AnimNext
 			}
 		}
 	}
+
+	void FMotionMatchingTrait::AddReferencedObjects(const FExecutionContext& Context, const TTraitBinding<IGarbageCollection>& Binding, FReferenceCollector& Collector) const
+	{
+		IGarbageCollection::AddReferencedObjects(Context, Binding, Collector);
+
+		FInstanceData* InstanceData = Binding.GetInstanceData<FInstanceData>();
+		FPoseSearchBlueprintResult& SearchResult = InstanceData->MotionMatchingState.SearchResult;
+
+		Collector.AddReferencedObject(SearchResult.SelectedAnim);
+		Collector.AddReferencedObject(SearchResult.SelectedDatabase);
+		Collector.AddReferencedObject(SearchResult.SelectedDatabase);
+
+		for (TObjectPtr<const UObject>& AnimContexts : SearchResult.AnimContexts)
+		{
+			Collector.AddReferencedObject(AnimContexts);
+		}
+	}
+
 } // UE::AnimNext
 
 void FAnimNextMotionMatchingTask::Execute(UE::AnimNext::FEvaluationVM& VM) const
diff --git a/Engine/Plugins/Experimental/AnimNextPoseSearch/Source/AnimNextPoseSearch/Private/MotionMatchingTrait.h b/Engine/Plugins/Experimental/AnimNextPoseSearch/Source/AnimNextPoseSearch/Private/MotionMatchingTrait.h
--- a/Engine/Plugins/Experimental/AnimNextPoseSearch/Source/AnimNextPoseSearch/Private/MotionMatchingTrait.h
+++ b/Engine/Plugins/Experimental/AnimNextPoseSearch/Source/AnimNextPoseSearch/Private/MotionMatchingTrait.h
@@ -9,6 +9,7 @@
 #include "TraitCore/Trait.h"
 #include "TraitCore/TraitSharedData.h"
 #include "TraitInterfaces/IEvaluate.h"
+#include "TraitInterfaces/IGarbageCollection.h"
 #include "TraitInterfaces/IGroupSynchronization.h"
 #include "TraitInterfaces/IUpdate.h"
 #include "MotionMatchingTrait.generated.h"
@@ -137,7 +138,7 @@ struct FMotionMatchingTraitSharedData : public FAnimNextTraitSharedData
 namespace UE::AnimNext
 {
 
-struct FMotionMatchingTrait : FAdditiveTrait, IUpdate, IEvaluate
+struct FMotionMatchingTrait : FAdditiveTrait, IUpdate, IEvaluate, IGarbageCollection
 {
 	DECLARE_ANIM_TRAIT(FMotionMatchingTrait, FAdditiveTrait)
 
@@ -159,6 +160,9 @@ struct FMotionMatchingTrait : FAdditiveTrait, IUpdate, IEvaluate
 	// IEvaluate impl
 	virtual void PostEvaluate(FEvaluateTraversalContext& Context, const TTraitBinding<IEvaluate>& Binding) const override;
 
+	// IGarbageCollection impl
+	virtual void AddReferencedObjects(const FExecutionContext& Context, const TTraitBinding<IGarbageCollection>& Binding, FReferenceCollector& Collector) const override;
+
 private:
 	void PublishResults(const TTraitBinding<IUpdate>& Binding) const;
 };
diff --git a/Engine/Plugins/Experimental/AnimNextPoseSearch/Source/AnimNextPoseSearch/Private/RigUnit_PoseSearchResultGetSelectedDatabase.cpp b/Engine/Plugins/Experimental/AnimNextPoseSearch/Source/AnimNextPoseSearch/Private/RigUnit_PoseSearchResultGetSelectedDatabase.cpp
--- a/Engine/Plugins/Experimental/AnimNextPoseSearch/Source/AnimNextPoseSearch/Private/RigUnit_PoseSearchResultGetSelectedDatabase.cpp
+++ b/Engine/Plugins/Experimental/AnimNextPoseSearch/Source/AnimNextPoseSearch/Private/RigUnit_PoseSearchResultGetSelectedDatabase.cpp
@@ -7,7 +7,7 @@
 
 FRigUnit_PoseSearchResultGetSelectedDatabase_Execute()
 {
-	if (PoseSearchResult.SelectedDatabase.IsValid())
+	if (PoseSearchResult.SelectedDatabase)
 	{
 		OutDatabase = const_cast<UPoseSearchDatabase*>(PoseSearchResult.SelectedDatabase.Get());
 	}
