fix for crash in pose history trait w/ stitches in sequencer

diff --git a/Engine/Plugins/Experimental/AnimNextPoseSearch/Source/AnimNextPoseSearch/Private/HistoryCollectorTrait.cpp b/Engine/Plugins/Experimental/AnimNextPoseSearch/Source/AnimNextPoseSearch/Private/HistoryCollectorTrait.cpp
--- a/Engine/Plugins/Experimental/AnimNextPoseSearch/Source/AnimNextPoseSearch/Private/HistoryCollectorTrait.cpp
+++ b/Engine/Plugins/Experimental/AnimNextPoseSearch/Source/AnimNextPoseSearch/Private/HistoryCollectorTrait.cpp
@@ -88,7 +88,7 @@ void FAnimNextHistoryCollectorPreEvaluateTask::Execute(UE::AnimNext::FEvaluation
 	if (EnumHasAnyFlags(VM.GetFlags(), EEvaluationFlags::Bones))
 	{
 		FPoseHistoryEvaluationHelper PoseHistoryEvalHelper;
-		PoseHistoryEvalHelper.PoseHistory = PoseHistory;
+		PoseHistoryEvalHelper.PoseHistoryPtr = InstanceData->PoseHistoryPtr;
 		VM.PushValue(POSEHISTORY_STACK_NAME, MakeUnique<FPoseHistoryEvaluationHelper>(PoseHistoryEvalHelper));
 	}
 }
@@ -218,8 +218,7 @@ namespace UE::AnimNext
 
 		check(InstanceData->PoseHistoryPtr.IsValid());
 		FAnimNextHistoryCollectorPreEvaluateTask Task;
-		Task.PoseHistory = InstanceData->PoseHistoryPtr.Get();
-		Task.HostObject = Context.GetHostObject();
+		Task.InstanceData = InstanceData;
 		Context.AppendTask(Task);
 	}
 
diff --git a/Engine/Plugins/Experimental/AnimNextPoseSearch/Source/AnimNextPoseSearch/Private/HistoryCollectorTrait.h b/Engine/Plugins/Experimental/AnimNextPoseSearch/Source/AnimNextPoseSearch/Private/HistoryCollectorTrait.h
--- a/Engine/Plugins/Experimental/AnimNextPoseSearch/Source/AnimNextPoseSearch/Private/HistoryCollectorTrait.h
+++ b/Engine/Plugins/Experimental/AnimNextPoseSearch/Source/AnimNextPoseSearch/Private/HistoryCollectorTrait.h
@@ -113,20 +113,6 @@ struct FAnimNextHistoryCollectorTraitSharedData : public FAnimNextTraitSharedDat
 	#undef TRAIT_LATENT_PROPERTIES_ENUMERATOR
 };
 
-USTRUCT()
-struct FAnimNextHistoryCollectorPreEvaluateTask : public FAnimNextEvaluationTask
-{
-	GENERATED_BODY()
-
-	DECLARE_ANIM_EVALUATION_TASK(FAnimNextHistoryCollectorPreEvaluateTask)
-
-	virtual void Execute(UE::AnimNext::FEvaluationVM& VM) const override;
-
-	UE::PoseSearch::FPoseHistory* PoseHistory = nullptr;
-	const FAnimNextHistoryCollectorTraitSharedData* SharedData = nullptr;
-	TWeakObjectPtr<const UObject> HostObject;
-};
-
 USTRUCT()
 struct FAnimNextHistoryCollectorTask : public FAnimNextEvaluationTask
 {
@@ -177,3 +163,14 @@ namespace UE::AnimNext
 		virtual const UE::PoseSearch::IPoseHistory* GetPoseHistory(FExecutionContext& Context, const TTraitBinding<IPoseHistory>& Binding) const override;
 	};
 }
+
+USTRUCT()
+struct FAnimNextHistoryCollectorPreEvaluateTask : public FAnimNextEvaluationTask
+{
+	GENERATED_BODY()
+
+	DECLARE_ANIM_EVALUATION_TASK(FAnimNextHistoryCollectorPreEvaluateTask)
+
+	virtual void Execute(UE::AnimNext::FEvaluationVM& VM) const override;
+	UE::AnimNext::FHistoryCollectorTrait::FInstanceData* InstanceData = nullptr;
+};
\ No newline at end of file
diff --git a/Engine/Plugins/Experimental/AnimNextPoseSearch/Source/AnimNextPoseSearch/Public/PoseHistoryEvaluation.h b/Engine/Plugins/Experimental/AnimNextPoseSearch/Source/AnimNextPoseSearch/Public/PoseHistoryEvaluation.h
--- a/Engine/Plugins/Experimental/AnimNextPoseSearch/Source/AnimNextPoseSearch/Public/PoseHistoryEvaluation.h
+++ b/Engine/Plugins/Experimental/AnimNextPoseSearch/Source/AnimNextPoseSearch/Public/PoseHistoryEvaluation.h
@@ -11,7 +11,7 @@ namespace UE::PoseSearch
 
 	struct FPoseHistoryEvaluationHelper
 	{
-		FPoseHistory* PoseHistory = nullptr;
+		TSharedPtr<UE::PoseSearch::FGenerateTrajectoryPoseHistory, ESPMode::ThreadSafe> PoseHistoryPtr = nullptr;
 	};
 }
 
diff --git a/Engine/Plugins/Experimental/MovieScenePoseSearchTracks/Source/MovieScenePoseSearchTracks/Private/Sections/MovieSceneStitchAnimSection.cpp b/Engine/Plugins/Experimental/MovieScenePoseSearchTracks/Source/MovieScenePoseSearchTracks/Private/Sections/MovieSceneStitchAnimSection.cpp
--- a/Engine/Plugins/Experimental/MovieScenePoseSearchTracks/Source/MovieScenePoseSearchTracks/Private/Sections/MovieSceneStitchAnimSection.cpp
+++ b/Engine/Plugins/Experimental/MovieScenePoseSearchTracks/Source/MovieScenePoseSearchTracks/Private/Sections/MovieSceneStitchAnimSection.cpp
@@ -64,7 +64,7 @@ void FMovieSceneStitchAnimEvaluationTask::Execute(UE::AnimNext::FEvaluationVM& V
 	const IPoseHistory* PoseHistory = nullptr;
 	if (const TUniquePtr<FPoseHistoryEvaluationHelper>* PoseHistoryEvalHelper = VM.PeekValue<TUniquePtr<FPoseHistoryEvaluationHelper>>(POSEHISTORY_STACK_NAME, 0))
 	{
-		PoseHistory = (*PoseHistoryEvalHelper)->PoseHistory;
+		PoseHistory = (*PoseHistoryEvalHelper)->PoseHistoryPtr.Get();
 	}
 
 	// If we haven't yet run a search and found an animation, do that
