PoseSearch - fix for MM trait bKeepInteractionAlive letting the interaction animation run until its length minus blend time (to avoid having to blend from a frozen animation that reached its end for the entire duration of the blend)

diff --git a/Engine/Plugins/Experimental/AnimNextPoseSearch/Source/AnimNextPoseSearch/Private/MotionMatchingTrait.cpp b/Engine/Plugins/Experimental/AnimNextPoseSearch/Source/AnimNextPoseSearch/Private/MotionMatchingTrait.cpp
--- a/Engine/Plugins/Experimental/AnimNextPoseSearch/Source/AnimNextPoseSearch/Private/MotionMatchingTrait.cpp
+++ b/Engine/Plugins/Experimental/AnimNextPoseSearch/Source/AnimNextPoseSearch/Private/MotionMatchingTrait.cpp
@@ -219,8 +219,10 @@ namespace UE::AnimNext
 				}
 				else
 				{
+					// letting the interaction animation run until its length minus blend time
+					// (to avoid having to blend from a frozen animation that reached its end for the entire duration of the blend)
 					const UAnimationAsset* AnimationAssetForRole = MotionMatchingState.SearchResult.GetAnimationAssetForRole();
-					if (!AnimationAssetForRole || MotionMatchingState.SearchResult.SelectedTime > AnimationAssetForRole->GetPlayLength())
+					if (!AnimationAssetForRole || MotionMatchingState.SearchResult.SelectedTime >= (AnimationAssetForRole->GetPlayLength() - SharedData->BlendArguments.BlendTime))
 					{
 						MotionMatchingState.SearchResult = FPoseSearchBlueprintResult();
 					}
