Fix for frozen animation when blending out of interaction that's being kept alive.

diff --git a/Engine/Plugins/Animation/PoseSearch/Source/Runtime/Public/PoseSearch/PoseSearchResult.h b/Engine/Plugins/Animation/PoseSearch/Source/Runtime/Public/PoseSearch/PoseSearchResult.h
--- a/Engine/Plugins/Animation/PoseSearch/Source/Runtime/Public/PoseSearch/PoseSearchResult.h
+++ b/Engine/Plugins/Animation/PoseSearch/Source/Runtime/Public/PoseSearch/PoseSearchResult.h
@@ -140,7 +140,7 @@ struct POSESEARCH_API FPoseSearchBlueprintResult
 	
 	// SelectedAnim suggested play rate (it's usually calculated by estimating the ratio between the trajectory speed in the query against the speed in the result)
 	UPROPERTY(Transient, VisibleAnywhere, BlueprintReadOnly, Category=State)
-	float WantedPlayRate = 0.f;
+	float WantedPlayRate = 1.f;
 
 	// SelectedAnim associated looping state
 	UPROPERTY(Transient, VisibleAnywhere, BlueprintReadOnly, Category=State)
diff --git a/Engine/Plugins/Experimental/AnimNextPoseSearch/Source/AnimNextPoseSearch/Private/MotionMatchingTrait.cpp b/Engine/Plugins/Experimental/AnimNextPoseSearch/Source/AnimNextPoseSearch/Private/MotionMatchingTrait.cpp
--- a/Engine/Plugins/Experimental/AnimNextPoseSearch/Source/AnimNextPoseSearch/Private/MotionMatchingTrait.cpp
+++ b/Engine/Plugins/Experimental/AnimNextPoseSearch/Source/AnimNextPoseSearch/Private/MotionMatchingTrait.cpp
@@ -379,7 +379,8 @@ namespace UE::AnimNext
 		IBlendStack::FGraphRequestPtr ActiveGraphRequest;
 		BlendStackBinding.GetActiveGraph(Context, ActiveGraphRequest);
 
-		if (ActiveGraphRequest)
+		// Only override current graph's inputs if we have a valid search result.
+		if (ActiveGraphRequest && MotionMatchingState.SearchResult.SelectedAnim)
 		{
 			check(ActiveGraphRequest->AnimationGraph);
 			for (FStructView& Payload : ActiveGraphRequest->GraphPayload.Get())
