Add Alignment Transform output to multi character motion matching system, and make choosers inside pose search databases work in Anim Next

diff --git a/Engine/Plugins/Animation/PoseSearch/Source/Runtime/Private/PoseSearchInteractionUtils.cpp b/Engine/Plugins/Animation/PoseSearch/Source/Runtime/Private/PoseSearchInteractionUtils.cpp
--- a/Engine/Plugins/Animation/PoseSearch/Source/Runtime/Private/PoseSearchInteractionUtils.cpp
+++ b/Engine/Plugins/Animation/PoseSearch/Source/Runtime/Private/PoseSearchInteractionUtils.cpp
@@ -37,7 +37,7 @@ FRoleToIndex MakeRoleToIndex(const UMultiAnimAsset* MultiAnimAsset)
 	return RoleToIndex;
 }
 
-void CalculateFullAlignedTransforms(const FPoseSearchBlueprintResult& CurrentResult, bool bWarpUsingRootBone, TArrayView<FTransform> OutFullAlignedTransforms)
+void CalculateFullAlignedTransformsAtTime(const FPoseSearchBlueprintResult& CurrentResult, float Time, bool bWarpUsingRootBone, TArrayView<FTransform> OutFullAlignedTransforms)
 {
 	const UMultiAnimAsset* MultiAnimAsset = CastChecked<UMultiAnimAsset>(CurrentResult.SelectedAnim);
 	const int32 NumRoles = MultiAnimAsset->GetNumRoles();
@@ -81,7 +81,12 @@ void CalculateFullAlignedTransforms(const FPoseSearchBlueprintResult& CurrentRes
 	{
 		RelevantRoleIndexes[Index] = CurrentResult.AnimContexts[Index] != nullptr;
 	}
-	MultiAnimAsset->CalculateWarpTransforms(CurrentResult.SelectedTime, ActorTransforms, OutFullAlignedTransforms, MirrorDataTables, RelevantRoleIndexes);
+	MultiAnimAsset->CalculateWarpTransforms(Time, ActorTransforms, OutFullAlignedTransforms, MirrorDataTables, RelevantRoleIndexes);
+}
+
+void CalculateFullAlignedTransforms(const FPoseSearchBlueprintResult& CurrentResult, bool bWarpUsingRootBone, TArrayView<FTransform> OutFullAlignedTransforms)
+{
+	CalculateFullAlignedTransformsAtTime(CurrentResult, CurrentResult.SelectedTime, bWarpUsingRootBone, OutFullAlignedTransforms);
 }
 
 FTransform CalculateDeltaAlignment(const FTransform& MeshWithoutOffset, const FTransform& MeshWithOffset, const FTransform& FullAlignedTransform, float WarpingRotationRatio, float WarpingTranslationRatio)
diff --git a/Engine/Plugins/Animation/PoseSearch/Source/Runtime/Public/PoseSearch/PoseSearchInteractionUtils.h b/Engine/Plugins/Animation/PoseSearch/Source/Runtime/Public/PoseSearch/PoseSearchInteractionUtils.h
--- a/Engine/Plugins/Animation/PoseSearch/Source/Runtime/Public/PoseSearch/PoseSearchInteractionUtils.h
+++ b/Engine/Plugins/Animation/PoseSearch/Source/Runtime/Public/PoseSearch/PoseSearchInteractionUtils.h
@@ -15,6 +15,9 @@ POSESEARCH_API int32 GetRoleIndex(const UMultiAnimAsset* MultiAnimAsset, const F
 
 // Experimental, this feature might be removed without warning, not for production use
 FRoleToIndex MakeRoleToIndex(const UMultiAnimAsset* MultiAnimAsset);
+	
+// Experimental, this feature might be removed without warning, not for production use
+POSESEARCH_API void CalculateFullAlignedTransformsAtTime(const FPoseSearchBlueprintResult& CurrentResult, float Time, bool bWarpUsingRootBone, TArrayView<FTransform> OutFullAlignedTransforms);
 
 // Experimental, this feature might be removed without warning, not for production use
 POSESEARCH_API void CalculateFullAlignedTransforms(const FPoseSearchBlueprintResult& CurrentResult, bool bWarpUsingRootBone, TArrayView<FTransform> OutFullAlignedTransforms);
diff --git a/Engine/Plugins/Experimental/AnimNextPoseSearch/Source/AnimNextPoseSearch/Private/MotionMatchingTrait.cpp b/Engine/Plugins/Experimental/AnimNextPoseSearch/Source/AnimNextPoseSearch/Private/MotionMatchingTrait.cpp
--- a/Engine/Plugins/Experimental/AnimNextPoseSearch/Source/AnimNextPoseSearch/Private/MotionMatchingTrait.cpp
+++ b/Engine/Plugins/Experimental/AnimNextPoseSearch/Source/AnimNextPoseSearch/Private/MotionMatchingTrait.cpp
@@ -259,6 +259,9 @@ namespace UE::AnimNext
 #endif // ENABLE_ANIM_DEBUG
 
 				FChooserEvaluationContext EvaluationContext(const_cast<UObject*>(AnimContext));
+				const FAnimNextDataInterfaceInstance& Instance = Binding.GetTraitPtr().GetNodeInstance()->GetOwner();
+                EvaluationContext.AddStructViewParam(	FStructView::Make(const_cast<FAnimNextDataInterfaceInstance&>(Instance)));
+                
 				UPoseSearchLibrary::UpdateMotionMatchingState(
 					&EvaluationContext,
 					PoseHistory,
@@ -407,6 +410,26 @@ namespace UE::AnimNext
 
 		PublishResults(Binding);
 
+		if (bIsInteraction)
+		{
+			if (const UMultiAnimAsset* MultiAnimAsset = Cast<UMultiAnimAsset>(MotionMatchingState.SearchResult.SelectedAnim))
+			{
+				const int32 CurrentResultRoleIndex = GetRoleIndex(MultiAnimAsset, MotionMatchingState.SearchResult.Role);
+				if (CurrentResultRoleIndex != INDEX_NONE)
+				{
+					const int32 NumRoles = MultiAnimAsset->GetNumRoles();
+					TArray<FTransform, TInlineAllocator<PreallocatedRolesNum>> FullAlignedTransforms;
+					FullAlignedTransforms.SetNum(NumRoles);
+
+					CalculateFullAlignedTransformsAtTime(MotionMatchingState.SearchResult, MultiAnimAsset->GetPlayLength(FVector::ZeroVector), SharedData->GetbWarpUsingRootBone(Binding), FullAlignedTransforms);
+					
+					FAnimNextDataInterfaceInstance& DataInterfaceInstance = Binding.GetTraitPtr().GetNodeInstance()->GetOwner();
+					FName AlignmentTransformName = SharedData->GetAlignmentTransformVariableName(Binding);
+					DataInterfaceInstance.SetVariable(AlignmentTransformName, FullAlignedTransforms[CurrentResultRoleIndex]);
+				}
+			}
+		}
+
 		// Motion Matching must finish updating before smoothers, otherwise we may request a blend after the smoother is done normalizing weights.
 		IUpdate::PreUpdate(Context, Binding, TraitState);
 	}
diff --git a/Engine/Plugins/Experimental/AnimNextPoseSearch/Source/AnimNextPoseSearch/Private/MotionMatchingTrait.h b/Engine/Plugins/Experimental/AnimNextPoseSearch/Source/AnimNextPoseSearch/Private/MotionMatchingTrait.h
--- a/Engine/Plugins/Experimental/AnimNextPoseSearch/Source/AnimNextPoseSearch/Private/MotionMatchingTrait.h
+++ b/Engine/Plugins/Experimental/AnimNextPoseSearch/Source/AnimNextPoseSearch/Private/MotionMatchingTrait.h
@@ -83,6 +83,9 @@ struct FMotionMatchingTraitSharedData : public FAnimNextTraitSharedData
 	// If set, the motion matching result will be written to this variable
 	UPROPERTY(EditAnywhere, Category = Settings)
 	FName MotionMatchingResultVariableName = NAME_None;
+	
+	UPROPERTY(EditAnywhere, Category = Settings)
+	FName AlignmentTransformVariableName;
 
 	// If set along with result variable name, the motion matching result will be written to the variable in this module, otherwise the variable will be written in the module containing the MotionMatchingTrait
 	UPROPERTY(EditAnywhere, Category = Settings)
@@ -130,6 +133,7 @@ struct FMotionMatchingTraitSharedData : public FAnimNextTraitSharedData
 		GeneratorMacro(SyncMode) \
 		GeneratorMacro(NextUpdateInterruptMode) \
 		GeneratorMacro(MotionMatchingResultVariableName) \
+		GeneratorMacro(AlignmentTransformVariableName) \
 		GeneratorMacro(MotionMatchingResultModuleHandle) \
 		GeneratorMacro(Availabilities) \
 		GeneratorMacro(bValidateResultAgainstAvailabilities) \
