Fix latent pin warning for PoseHistoryReferenceVariable

diff --git a/Engine/Plugins/Experimental/AnimNextPoseSearch/Source/AnimNextPoseSearch/Private/HistoryCollectorTrait.cpp b/Engine/Plugins/Experimental/AnimNextPoseSearch/Source/AnimNextPoseSearch/Private/HistoryCollectorTrait.cpp
--- a/Engine/Plugins/Experimental/AnimNextPoseSearch/Source/AnimNextPoseSearch/Private/HistoryCollectorTrait.cpp
+++ b/Engine/Plugins/Experimental/AnimNextPoseSearch/Source/AnimNextPoseSearch/Private/HistoryCollectorTrait.cpp
@@ -241,12 +241,13 @@ namespace UE::AnimNext
 		}
 #endif // WITH_EDITOR
 
-		if (!SharedData->PoseHistoryReferenceVariable.IsNone())
+		FName PoseHistoryReferenceVariableName = SharedData->GetPoseHistoryReferenceVariable(Binding);
+		if (!PoseHistoryReferenceVariableName.IsNone())
 		{
 			FAnimNextModuleInstance* ModuleInstance = Binding.GetTraitPtr().GetNodeInstance()->GetOwner().GetModuleInstance();
 			FPoseHistoryReference Reference;
 			Reference.PoseHistory = InstanceData->PoseHistoryPtr;
-			ModuleInstance->SetVariable(SharedData->PoseHistoryReferenceVariable, Reference);
+			ModuleInstance->SetVariable(PoseHistoryReferenceVariableName, Reference);
 		}
 		
 		check(InstanceData->PoseHistoryPtr.IsValid());
diff --git a/Engine/Plugins/Experimental/AnimNextPoseSearch/Source/AnimNextPoseSearch/Private/HistoryCollectorTrait.h b/Engine/Plugins/Experimental/AnimNextPoseSearch/Source/AnimNextPoseSearch/Private/HistoryCollectorTrait.h
--- a/Engine/Plugins/Experimental/AnimNextPoseSearch/Source/AnimNextPoseSearch/Private/HistoryCollectorTrait.h
+++ b/Engine/Plugins/Experimental/AnimNextPoseSearch/Source/AnimNextPoseSearch/Private/HistoryCollectorTrait.h
@@ -108,6 +108,7 @@ struct FAnimNextHistoryCollectorTraitSharedData : public FAnimNextTraitSharedDat
 		GeneratorMacro(TrajectoryHistoryCount) \
 		GeneratorMacro(TrajectoryPredictionCount) \
 		GeneratorMacro(PredictionSamplingInterval) \
+		GeneratorMacro(PoseHistoryReferenceVariable) \
 
 	GENERATE_TRAIT_LATENT_PROPERTIES(FAnimNextHistoryCollectorTraitSharedData, TRAIT_LATENT_PROPERTIES_ENUMERATOR)
 	#undef TRAIT_LATENT_PROPERTIES_ENUMERATOR
