PoseSearch - NonUnity fix

diff --git a/Engine/Plugins/Experimental/AnimNextPoseSearch/Source/AnimNextPoseSearch/Private/PoseSearchResultEmulatorTrait.h b/Engine/Plugins/Experimental/AnimNextPoseSearch/Source/AnimNextPoseSearch/Private/PoseSearchResultEmulatorTrait.h
--- a/Engine/Plugins/Experimental/AnimNextPoseSearch/Source/AnimNextPoseSearch/Private/PoseSearchResultEmulatorTrait.h
+++ b/Engine/Plugins/Experimental/AnimNextPoseSearch/Source/AnimNextPoseSearch/Private/PoseSearchResultEmulatorTrait.h
@@ -2,12 +2,13 @@
 
 #pragma once
 
+#include "AlphaBlend.h"
 #include "CoreMinimal.h"
 #include "PoseSearch/PoseSearchLibrary.h"
 #include "TraitCore/Trait.h"
 #include "TraitCore/TraitSharedData.h"
-#include "TraitInterfaces/IUpdate.h"
 #include "TraitInterfaces/IEvaluate.h"
+#include "TraitInterfaces/IUpdate.h"
 #include "PoseSearchResultEmulatorTrait.generated.h"
 
 USTRUCT(meta = (DisplayName = "Pose Search Result Emulator"))
