AnimNextPoseSearch - fixes to make FMotionMatchingTrait GC compliant

diff --git a/Engine/Plugins/Experimental/AnimNextPoseSearch/Source/AnimNextPoseSearch/Private/MotionMatchingTrait.cpp b/Engine/Plugins/Experimental/AnimNextPoseSearch/Source/AnimNextPoseSearch/Private/MotionMatchingTrait.cpp
--- a/Engine/Plugins/Experimental/AnimNextPoseSearch/Source/AnimNextPoseSearch/Private/MotionMatchingTrait.cpp
+++ b/Engine/Plugins/Experimental/AnimNextPoseSearch/Source/AnimNextPoseSearch/Private/MotionMatchingTrait.cpp
@@ -60,6 +60,18 @@ namespace UE::AnimNext
 	GENERATE_ANIM_TRAIT_IMPLEMENTATION(FMotionMatchingTrait, TRAIT_INTERFACE_ENUMERATOR, NULL_ANIM_TRAIT_INTERFACE_ENUMERATOR, NULL_ANIM_TRAIT_EVENT_ENUMERATOR)
 #undef TRAIT_INTERFACE_ENUMERATOR
 
+	void FMotionMatchingTrait::FInstanceData::Construct(const FExecutionContext& Context, const FTraitBinding& Binding)
+	{
+		FTrait::FInstanceData::Construct(Context, Binding);
+		IGarbageCollection::RegisterWithGC(Context, Binding);
+	}
+
+	void FMotionMatchingTrait::FInstanceData::Destruct(const FExecutionContext& Context, const FTraitBinding& Binding)
+	{
+		FTrait::FInstanceData::Destruct(Context, Binding);
+		IGarbageCollection::UnregisterWithGC(Context, Binding);
+	}
+
 	void FMotionMatchingTrait::OnBecomeRelevant(FUpdateTraversalContext& Context, const TTraitBinding<IUpdate>& Binding, const FTraitUpdateState& TraitState) const
 	{
 		IUpdate::OnBecomeRelevant(Context, Binding, TraitState);
@@ -439,15 +451,16 @@ namespace UE::AnimNext
 		using namespace UE::PoseSearch;
 
 		const FSharedData* SharedData = Binding.GetSharedData<FSharedData>();
-		const FInstanceData* InstanceData = Binding.GetInstanceData<FInstanceData>();
-		const FMotionMatchingState& MotionMatchingState = InstanceData->MotionMatchingState;
+		const FName ResultOutputVariableName = SharedData->GetMotionMatchingResultVariableName(Binding);
+		const FName ResultOutputVariableNameAlt = SharedData->GetMotionMatchingResultVariableNameAlt(Binding);
 
-		FName ResultOutputVariableName = SharedData->GetMotionMatchingResultVariableName(Binding);
-		if (ResultOutputVariableName == NAME_None)
+		if (ResultOutputVariableName == NAME_None && ResultOutputVariableNameAlt == NAME_None)
 		{
 			return;
 		}
 
+		const FInstanceData* InstanceData = Binding.GetInstanceData<FInstanceData>();
+		const FMotionMatchingState& MotionMatchingState = InstanceData->MotionMatchingState;
 		const UPoseSearchDatabase* CurrentResultDatabase = MotionMatchingState.SearchResult.SelectedDatabase.Get();
 		if (CurrentResultDatabase == nullptr || CurrentResultDatabase->Schema == nullptr)
 		{
@@ -460,30 +473,72 @@ namespace UE::AnimNext
 		}
 #endif // WITH_EDITOR
 		
+		const auto PublishResult = [](FAnimNextModuleInstance* ModuleInstance, const FModuleHandle& ModuleHandle, const FPoseSearchBlueprintResult& SearchResult, FName ResultOutputVariableName)
+			{
+				if (ModuleHandle.IsValid() && ResultOutputVariableName != NAME_None)
+				{
+					// storing SearchResult TObjectPtr(s) as TWeakObjectPtr(s) to avoid referencing deallocated memory 
+					// in case the module containing this trait gets destroyed before the QueueTaskOnOtherModule lambda execution
+					TWeakObjectPtr<UObject> SelectedAnim = SearchResult.SelectedAnim;
+					float SelectedTime = SearchResult.SelectedTime;
+					bool bIsContinuingPoseSearch = SearchResult.bIsContinuingPoseSearch;
+					float WantedPlayRate = SearchResult.WantedPlayRate;
+					bool bLoop = SearchResult.bLoop;
+					bool bIsMirrored = SearchResult.bIsMirrored;
+					FVector BlendParameters = SearchResult.BlendParameters;
+					TWeakObjectPtr<const UPoseSearchDatabase> SelectedDatabase = SearchResult.SelectedDatabase;
+					float SearchCost = SearchResult.SearchCost;
+					bool bIsInteraction = SearchResult.bIsInteraction;
+					FName Role = SearchResult.Role;
+					TArray<FTransform> ActorRootTransforms = SearchResult.ActorRootTransforms;
+					TArray<FTransform> ActorRootBoneTransforms = SearchResult.ActorRootBoneTransforms;
+
+					ModuleInstance->QueueTaskOnOtherModule(
+						ModuleHandle,
+						FRigUnit_AnimNextPrePhysicsEvent::DefaultEventName,
+						[SelectedAnim, SelectedTime, bIsContinuingPoseSearch, WantedPlayRate, bLoop, 
+						bIsMirrored, BlendParameters, SelectedDatabase, SearchCost, bIsInteraction,
+						Role, ActorRootTransforms, ActorRootBoneTransforms, ResultOutputVariableName](const UE::AnimNext::FModuleTaskContext& TaskContext)
+						{
+							FAnimNextModuleInstance* OtherModuleInstance = TaskContext.GetModuleInstance();
+							if (OtherModuleInstance)
+							{
+								// reconstructing the FPoseSearchBlueprintResult from the input properties
+								FPoseSearchBlueprintResult SearchResult;
+								SearchResult.SelectedAnim = SelectedAnim.Get();
+								SearchResult.SelectedTime = SelectedTime;
+								SearchResult.bIsContinuingPoseSearch = bIsContinuingPoseSearch;
+								SearchResult.WantedPlayRate = WantedPlayRate;
+								SearchResult.bLoop = bLoop;
+								SearchResult.bIsMirrored = bIsMirrored;
+								SearchResult.BlendParameters = BlendParameters;
+								SearchResult.SelectedDatabase = SelectedDatabase.Get();
+								SearchResult.SearchCost = SearchCost;
+								SearchResult.bIsInteraction = bIsInteraction;
+								SearchResult.Role = Role;
+								SearchResult.ActorRootTransforms = ActorRootTransforms;
+								SearchResult.ActorRootBoneTransforms = ActorRootBoneTransforms;
+								OtherModuleInstance->SetVariable(ResultOutputVariableName, SearchResult);
+							}
+						},
+						UE::AnimNext::ETaskRunLocation::Before);
+				}
+				else
+				{
+					FPoseSearchBlueprintResult SearchResultCopy = SearchResult;
+					SearchResultCopy.AnimContexts.Reset();
+					ModuleInstance->SetVariable(ResultOutputVariableName, SearchResultCopy);
+				}
+			};
+
 		FAnimNextModuleInstance* ModuleInstance = Binding.GetTraitPtr().GetNodeInstance()->GetOwner().GetModuleInstance();
 		check(ModuleInstance);
 
-		// @todo: merge this duplicated code with the block below
 		const FAnimNextModuleHandle ResultModuleHandle = SharedData->GetMotionMatchingResultModuleHandle(Binding);
-		if (ResultModuleHandle.ModuleHandle.IsValid())
-		{
-			ModuleInstance->QueueTaskOnOtherModule(
-				ResultModuleHandle.ModuleHandle,
-				FRigUnit_AnimNextPrePhysicsEvent::DefaultEventName,
-				[&MotionMatchingState, ResultOutputVariableName](const UE::AnimNext::FModuleTaskContext& TaskContext)
-				{
-					FAnimNextModuleInstance* OtherModuleInstance = TaskContext.GetModuleInstance();
-					if (OtherModuleInstance)
-					{
-						OtherModuleInstance->SetVariable(ResultOutputVariableName, MotionMatchingState.SearchResult);
-					}
-				},
-				UE::AnimNext::ETaskRunLocation::Before);
-		}
-		else
-		{
-			ModuleInstance->SetVariable(ResultOutputVariableName, MotionMatchingState.SearchResult);
-		}
+		PublishResult(ModuleInstance, ResultModuleHandle.ModuleHandle, MotionMatchingState.SearchResult, ResultOutputVariableName);
+		
+		const FAnimNextModuleHandle ResultModuleHandleAlt = SharedData->GetMotionMatchingResultModuleHandleAlt(Binding);
+		PublishResult(ModuleInstance, ResultModuleHandleAlt.ModuleHandle, MotionMatchingState.SearchResult, ResultOutputVariableNameAlt);
 	}
 
 	void FMotionMatchingTrait::PostEvaluate(FEvaluateTraversalContext& Context, const TTraitBinding<IEvaluate>& Binding) const
diff --git a/Engine/Plugins/Experimental/AnimNextPoseSearch/Source/AnimNextPoseSearch/Private/MotionMatchingTrait.h b/Engine/Plugins/Experimental/AnimNextPoseSearch/Source/AnimNextPoseSearch/Private/MotionMatchingTrait.h
--- a/Engine/Plugins/Experimental/AnimNextPoseSearch/Source/AnimNextPoseSearch/Private/MotionMatchingTrait.h
+++ b/Engine/Plugins/Experimental/AnimNextPoseSearch/Source/AnimNextPoseSearch/Private/MotionMatchingTrait.h
@@ -91,6 +91,14 @@ struct FMotionMatchingTraitSharedData : public FAnimNextTraitSharedData
 	UPROPERTY(EditAnywhere, Category = Settings)
 	FAnimNextModuleHandle MotionMatchingResultModuleHandle;
 
+	// Same as MotionMatchingResultVariableName but for an alternative destination variable
+	UPROPERTY(EditAnywhere, Category = Settings)
+	FName MotionMatchingResultVariableNameAlt = NAME_None;
+
+	// Same as MotionMatchingResultModuleHandle but for an alternative destination variable
+	UPROPERTY(EditAnywhere, Category = Settings)
+	FAnimNextModuleHandle MotionMatchingResultModuleHandleAlt;
+
 	UPROPERTY(EditAnywhere, Category = Interaction, meta = (PinHiddenByDefault))
 	TArray<FPoseSearchInteractionAvailability> Availabilities;
 
@@ -135,6 +143,8 @@ struct FMotionMatchingTraitSharedData : public FAnimNextTraitSharedData
 		GeneratorMacro(MotionMatchingResultVariableName) \
 		GeneratorMacro(AlignmentTransformVariableName) \
 		GeneratorMacro(MotionMatchingResultModuleHandle) \
+		GeneratorMacro(MotionMatchingResultVariableNameAlt) \
+		GeneratorMacro(MotionMatchingResultModuleHandleAlt) \
 		GeneratorMacro(Availabilities) \
 		GeneratorMacro(bValidateResultAgainstAvailabilities) \
 		GeneratorMacro(bKeepInteractionAlive) \
@@ -164,6 +174,9 @@ struct FMotionMatchingTrait : FAdditiveTrait, IUpdate, IEvaluate, IGarbageCollec
 #if WITH_EDITOR
 		bool bIsPostEvaluateBeingCalled = true;
 #endif // WITH_EDITOR
+
+		void Construct(const FExecutionContext& Context, const FTraitBinding& Binding);
+		void Destruct(const FExecutionContext& Context, const FTraitBinding& Binding);
 	};
 
 	// IUpdate impl
