AnimNextPoseSearch - fix for FMotionMatchingTrait::PublishResults trying to call FAnimNextDataInterfaceInstance::SetVariable with NAME_None

diff --git a/Engine/Plugins/Experimental/AnimNextPoseSearch/Source/AnimNextPoseSearch/Private/MotionMatchingTrait.cpp b/Engine/Plugins/Experimental/AnimNextPoseSearch/Source/AnimNextPoseSearch/Private/MotionMatchingTrait.cpp
--- a/Engine/Plugins/Experimental/AnimNextPoseSearch/Source/AnimNextPoseSearch/Private/MotionMatchingTrait.cpp
+++ b/Engine/Plugins/Experimental/AnimNextPoseSearch/Source/AnimNextPoseSearch/Private/MotionMatchingTrait.cpp
@@ -475,59 +475,62 @@ namespace UE::AnimNext
 		
 		const auto PublishResult = [](FAnimNextModuleInstance* ModuleInstance, const FModuleHandle& ModuleHandle, const FPoseSearchBlueprintResult& SearchResult, FName ResultOutputVariableName)
 			{
-				if (ModuleHandle.IsValid() && ResultOutputVariableName != NAME_None)
+				if (ResultOutputVariableName != NAME_None)
 				{
-					// storing SearchResult TObjectPtr(s) as TWeakObjectPtr(s) to avoid referencing deallocated memory 
-					// in case the module containing this trait gets destroyed before the QueueTaskOnOtherModule lambda execution
-					TWeakObjectPtr<UObject> SelectedAnim = SearchResult.SelectedAnim;
-					float SelectedTime = SearchResult.SelectedTime;
-					bool bIsContinuingPoseSearch = SearchResult.bIsContinuingPoseSearch;
-					float WantedPlayRate = SearchResult.WantedPlayRate;
-					bool bLoop = SearchResult.bLoop;
-					bool bIsMirrored = SearchResult.bIsMirrored;
-					FVector BlendParameters = SearchResult.BlendParameters;
-					TWeakObjectPtr<const UPoseSearchDatabase> SelectedDatabase = SearchResult.SelectedDatabase;
-					float SearchCost = SearchResult.SearchCost;
-					bool bIsInteraction = SearchResult.bIsInteraction;
-					FName Role = SearchResult.Role;
-					TArray<FTransform> ActorRootTransforms = SearchResult.ActorRootTransforms;
-					TArray<FTransform> ActorRootBoneTransforms = SearchResult.ActorRootBoneTransforms;
-
-					ModuleInstance->QueueTaskOnOtherModule(
-						ModuleHandle,
-						FRigUnit_AnimNextPrePhysicsEvent::DefaultEventName,
-						[SelectedAnim, SelectedTime, bIsContinuingPoseSearch, WantedPlayRate, bLoop, 
-						bIsMirrored, BlendParameters, SelectedDatabase, SearchCost, bIsInteraction,
-						Role, ActorRootTransforms, ActorRootBoneTransforms, ResultOutputVariableName](const UE::AnimNext::FModuleTaskContext& TaskContext)
-						{
-							FAnimNextModuleInstance* OtherModuleInstance = TaskContext.GetModuleInstance();
-							if (OtherModuleInstance)
+					if (ModuleHandle.IsValid())
+					{
+						// storing SearchResult TObjectPtr(s) as TWeakObjectPtr(s) to avoid referencing deallocated memory 
+						// in case the module containing this trait gets destroyed before the QueueTaskOnOtherModule lambda execution
+						TWeakObjectPtr<UObject> SelectedAnim = SearchResult.SelectedAnim;
+						float SelectedTime = SearchResult.SelectedTime;
+						bool bIsContinuingPoseSearch = SearchResult.bIsContinuingPoseSearch;
+						float WantedPlayRate = SearchResult.WantedPlayRate;
+						bool bLoop = SearchResult.bLoop;
+						bool bIsMirrored = SearchResult.bIsMirrored;
+						FVector BlendParameters = SearchResult.BlendParameters;
+						TWeakObjectPtr<const UPoseSearchDatabase> SelectedDatabase = SearchResult.SelectedDatabase;
+						float SearchCost = SearchResult.SearchCost;
+						bool bIsInteraction = SearchResult.bIsInteraction;
+						FName Role = SearchResult.Role;
+						TArray<FTransform> ActorRootTransforms = SearchResult.ActorRootTransforms;
+						TArray<FTransform> ActorRootBoneTransforms = SearchResult.ActorRootBoneTransforms;
+
+						ModuleInstance->QueueTaskOnOtherModule(
+							ModuleHandle,
+							FRigUnit_AnimNextPrePhysicsEvent::DefaultEventName,
+							[SelectedAnim, SelectedTime, bIsContinuingPoseSearch, WantedPlayRate, bLoop,
+							bIsMirrored, BlendParameters, SelectedDatabase, SearchCost, bIsInteraction,
+							Role, ActorRootTransforms, ActorRootBoneTransforms, ResultOutputVariableName](const UE::AnimNext::FModuleTaskContext& TaskContext)
 							{
-								// reconstructing the FPoseSearchBlueprintResult from the input properties
-								FPoseSearchBlueprintResult SearchResult;
-								SearchResult.SelectedAnim = SelectedAnim.Get();
-								SearchResult.SelectedTime = SelectedTime;
-								SearchResult.bIsContinuingPoseSearch = bIsContinuingPoseSearch;
-								SearchResult.WantedPlayRate = WantedPlayRate;
-								SearchResult.bLoop = bLoop;
-								SearchResult.bIsMirrored = bIsMirrored;
-								SearchResult.BlendParameters = BlendParameters;
-								SearchResult.SelectedDatabase = SelectedDatabase.Get();
-								SearchResult.SearchCost = SearchCost;
-								SearchResult.bIsInteraction = bIsInteraction;
-								SearchResult.Role = Role;
-								SearchResult.ActorRootTransforms = ActorRootTransforms;
-								SearchResult.ActorRootBoneTransforms = ActorRootBoneTransforms;
-								OtherModuleInstance->SetVariable(ResultOutputVariableName, SearchResult);
-							}
-						},
-						UE::AnimNext::ETaskRunLocation::Before);
-				}
-				else
-				{
-					FPoseSearchBlueprintResult SearchResultCopy = SearchResult;
-					SearchResultCopy.AnimContexts.Reset();
-					ModuleInstance->SetVariable(ResultOutputVariableName, SearchResultCopy);
+								FAnimNextModuleInstance* OtherModuleInstance = TaskContext.GetModuleInstance();
+								if (OtherModuleInstance)
+								{
+									// reconstructing the FPoseSearchBlueprintResult from the input properties
+									FPoseSearchBlueprintResult SearchResult;
+									SearchResult.SelectedAnim = SelectedAnim.Get();
+									SearchResult.SelectedTime = SelectedTime;
+									SearchResult.bIsContinuingPoseSearch = bIsContinuingPoseSearch;
+									SearchResult.WantedPlayRate = WantedPlayRate;
+									SearchResult.bLoop = bLoop;
+									SearchResult.bIsMirrored = bIsMirrored;
+									SearchResult.BlendParameters = BlendParameters;
+									SearchResult.SelectedDatabase = SelectedDatabase.Get();
+									SearchResult.SearchCost = SearchCost;
+									SearchResult.bIsInteraction = bIsInteraction;
+									SearchResult.Role = Role;
+									SearchResult.ActorRootTransforms = ActorRootTransforms;
+									SearchResult.ActorRootBoneTransforms = ActorRootBoneTransforms;
+									OtherModuleInstance->SetVariable(ResultOutputVariableName, SearchResult);
+								}
+							},
+							UE::AnimNext::ETaskRunLocation::Before);
+					}
+					else
+					{
+						FPoseSearchBlueprintResult SearchResultCopy = SearchResult;
+						SearchResultCopy.AnimContexts.Reset();
+						ModuleInstance->SetVariable(ResultOutputVariableName, SearchResultCopy);
+					}
 				}
 			};
 
