Fix for pop during interactions

diff --git a/Engine/Plugins/Experimental/AnimNextPoseSearch/Source/AnimNextPoseSearch/Private/MotionMatchingTrait.cpp b/Engine/Plugins/Experimental/AnimNextPoseSearch/Source/AnimNextPoseSearch/Private/MotionMatchingTrait.cpp
--- a/Engine/Plugins/Experimental/AnimNextPoseSearch/Source/AnimNextPoseSearch/Private/MotionMatchingTrait.cpp
+++ b/Engine/Plugins/Experimental/AnimNextPoseSearch/Source/AnimNextPoseSearch/Private/MotionMatchingTrait.cpp
@@ -681,8 +681,15 @@ void FAnimNextMotionMatchingTask::Execute(UE::AnimNext::FEvaluationVM& VM) const
 	const FTransform& MeshWithOffset = MeshWithoutOffset;
 
 	// NoTe: keep in mind DeltaAlignment is relative to the previous execution frame so we still need to extract and and apply the current animation root motion transform to get to the current frame full aligned transform.
-	const FTransform DeltaAlignment = CalculateDeltaAlignment(MeshWithoutOffset, MeshWithOffset, FullAlignedTransforms[CurrentResultRoleIndex], FinalWarpingRotationRatio, FinalWarpingTranslationRatio);
-	
+	FTransform DeltaAlignment = CalculateDeltaAlignment(MeshWithoutOffset, MeshWithOffset, FullAlignedTransforms[CurrentResultRoleIndex], FinalWarpingRotationRatio, FinalWarpingTranslationRatio);
+	// @TODO: This is a HACK. Since our warp is framerate dependent, if our characters are at different heights we may generate excessive z velocity resulting in pops at high frame rates.
+	{
+		FVector DeltaTranslation = DeltaAlignment.GetTranslation();
+		DeltaTranslation.Z = 0.0f;
+		DeltaAlignment.SetTranslation(DeltaTranslation);
+	}
+
+
 	const FTransform DeltaAlignmentWithRootMotion = DeltaAlignment * RootMotionDelta;
 
 	RootMotionProvider->OverrideRootMotion(DeltaAlignmentWithRootMotion, Keyframe->Get()->Attributes);
