Fix for desync between in pose search emulator after cinematic

diff --git a/Engine/Plugins/Experimental/AnimNextPoseSearch/Source/AnimNextPoseSearch/Private/PoseSearchResultEmulatorTrait.cpp b/Engine/Plugins/Experimental/AnimNextPoseSearch/Source/AnimNextPoseSearch/Private/PoseSearchResultEmulatorTrait.cpp
--- a/Engine/Plugins/Experimental/AnimNextPoseSearch/Source/AnimNextPoseSearch/Private/PoseSearchResultEmulatorTrait.cpp
+++ b/Engine/Plugins/Experimental/AnimNextPoseSearch/Source/AnimNextPoseSearch/Private/PoseSearchResultEmulatorTrait.cpp
@@ -48,10 +48,9 @@ namespace UE::AnimNext
 
 		const FSharedData* SharedData = Binding.GetSharedData<FSharedData>();
 		const UObject* AnimObject = SharedData->GetSelectedAnim(Binding);
+		const FPoseSearchBlueprintResult& Result = SharedData->GetPoseSearchResult(Binding);
 		if (AnimObject == nullptr)
 		{
-			// @TODO: Kill this with fire.
-			const FPoseSearchBlueprintResult& Result = SharedData->GetPoseSearchResult(Binding);
 			AnimObject = Result.SelectedAnim.Get();
 		}
 
@@ -104,7 +103,9 @@ namespace UE::AnimNext
 
 			const float MaxTimeDeltaAllowed = SharedData->GetMaxTimeDeltaAllowed(Binding);
 			const bool bDoBlend =	CurrentAnim != AnimationAsset
-								||	(TimeDelta > MaxTimeDeltaAllowed && !bIsBlendSpace);
+								||	TimeDelta > FMath::Max(MaxTimeDeltaAllowed, UE_SMALL_NUMBER)
+								// If it's not a continuing pose search, then our search requested a new pose.
+								|| !Result.bIsContinuingPoseSearch;
 
 			if (bDoBlend)
 			{
